GameplatStateTree Plugin
- moved StateTreeBrainComponent (from StateTreeModule) and renamed StateTreeComponent
- added FStateTreeTaskCommonBase (similar to FStateTreeTaskCommonCondition)
#rnx

diff --git a/Engine/Plugins/Runtime/GameplayStateTree/Config/BaseGameplayStateTree.ini b/Engine/Plugins/Runtime/GameplayStateTree/Config/BaseGameplayStateTree.ini
new file mode 100644
--- /dev/null
+++ b/Engine/Plugins/Runtime/GameplayStateTree/Config/BaseGameplayStateTree.ini
@@ -0,0 +1,3 @@
+﻿[CoreRedirects]
++ClassRedirects=(OldName="/Script/StateTreeModule.StateTreeBrainComponent", NewName="/Script/GameplayStateTreeModule.StateTreeComponent")
++ClassRedirects=(OldName="/Script/StateTreeModule.BrainComponentStateTreeSchema", NewName="/Script/GameplayStateTreeModule.StateTreeComponentSchema")
diff --git a/Engine/Plugins/Runtime/GameplayStateTree/GameplayStateTree.uplugin b/Engine/Plugins/Runtime/GameplayStateTree/GameplayStateTree.uplugin
new file mode 100644
--- /dev/null
+++ b/Engine/Plugins/Runtime/GameplayStateTree/GameplayStateTree.uplugin
@@ -0,0 +1,35 @@
+{
+	"FileVersion" : 3,
+	"Version" : 1,
+	"VersionName" : "1.0",
+	"FriendlyName" : "GameplayStateTree",
+	"Description" : "StateTree for AI/Gameplay Behaviors",
+	"Category" : "Gameplay",
+	"CreatedBy" : "Epic Games, Inc.",
+	"CreatedByURL" : "http://epicgames.com",
+	"DocsURL" : "",
+	"MarketplaceURL" : "",
+	"SupportURL" : "",
+	"EnabledByDefault" : false,
+	"CanContainContent" : true,
+	"IsBetaVersion" : false,
+	"IsExperimentalVersion": true,
+	"Installed" : false,
+	"Modules": [
+		{
+			"Name": "GameplayStateTreeModule",
+			"Type": "Runtime",
+			"LoadingPhase": "Default"
+		}
+	],
+	"Plugins": [
+		{
+			"Name": "StateTree",
+			"Enabled": true
+		},
+		{
+			"Name": "GameplayInteractions",
+			"Enabled": true
+		}
+	]
+}
\ No newline at end of file
diff --git a/Engine/Plugins/Runtime/GameplayStateTree/Source/GameplayStateTreeModule/GameplayStateTreeModule.Build.cs b/Engine/Plugins/Runtime/GameplayStateTree/Source/GameplayStateTreeModule/GameplayStateTreeModule.Build.cs
new file mode 100644
--- /dev/null
+++ b/Engine/Plugins/Runtime/GameplayStateTree/Source/GameplayStateTreeModule/GameplayStateTreeModule.Build.cs
@@ -0,0 +1,26 @@
+// Copyright Epic Games, Inc. All Rights Reserved.
+
+namespace UnrealBuildTool.Rules
+{
+	public class GameplayStateTreeModule : ModuleRules
+	{
+		public GameplayStateTreeModule(ReadOnlyTargetRules Target) : base(Target)
+		{
+			PublicIncludePaths.AddRange(
+				new string[] {
+				}
+			);
+
+			PublicDependencyModuleNames.AddRange(
+				new string[] {
+					"AIModule",
+					"Core",
+					"CoreUObject",
+					"Engine",
+					"GameplayTasks",
+					"StateTreeModule"
+				}
+			);
+		}
+	}
+}
diff --git a/Engine/Plugins/Runtime/GameplayStateTree/Source/GameplayStateTreeModule/Private/Components/StateTreeComponent.cpp b/Engine/Plugins/Runtime/GameplayStateTree/Source/GameplayStateTreeModule/Private/Components/StateTreeComponent.cpp
new file mode 100644
--- /dev/null
+++ b/Engine/Plugins/Runtime/GameplayStateTree/Source/GameplayStateTreeModule/Private/Components/StateTreeComponent.cpp
@@ -0,0 +1,249 @@
+// Copyright Epic Games, Inc. All Rights Reserved.
+
+#include "Components/StateTreeComponent.h"
+#include "GameFramework/Actor.h"
+#include "GameplayTasksComponent.h"
+#include "VisualLogger/VisualLogger.h"
+#include "StateTree.h"
+#include "StateTreeEvaluatorBase.h"
+#include "Conditions/StateTreeCondition_Common.h"
+#include "AIController.h"
+#include "GameFramework/CharacterMovementComponent.h"
+
+#define STATETREE_LOG(Verbosity, Format, ...) UE_VLOG(GetOwner(), LogStateTree, Verbosity, Format, ##__VA_ARGS__)
+#define STATETREE_CLOG(Condition, Verbosity, Format, ...) UE_CVLOG((Condition), GetOwner(), LogStateTree, Verbosity, Format, ##__VA_ARGS__)
+
+//////////////////////////////////////////////////////////////////////////
+// UStateTreeComponent
+
+UStateTreeComponent::UStateTreeComponent(const FObjectInitializer& ObjectInitializer)
+	: Super(ObjectInitializer)
+{
+	bWantsInitializeComponent = true;
+	bIsRunning = false;
+	bIsPaused = false;
+}
+
+void UStateTreeComponent::InitializeComponent()
+{
+	if (StateTree == nullptr)
+	{
+		STATETREE_LOG(Error, TEXT("%s: StateTree asset is not set, cannot initialize."), ANSI_TO_TCHAR(__FUNCTION__));
+		return;
+	}
+
+	if (!StateTreeContext.Init(*GetOwner(), *StateTree, EStateTreeStorage::Internal))
+	{
+		STATETREE_LOG(Error, TEXT("%s: Failed to init StateTreeContext."), ANSI_TO_TCHAR(__FUNCTION__));
+	}
+}
+
+void UStateTreeComponent::UninitializeComponent()
+{
+}
+
+bool UStateTreeComponent::SetContextRequirements()
+{
+	if (!StateTreeContext.IsValid())
+	{
+		return false;
+	}
+
+	const UWorld* World = GetWorld();
+	if (World == nullptr)
+	{
+		return false;
+	}
+
+	for (const FStateTreeExternalDataDesc& ItemDesc : StateTreeContext.GetExternalDataDescs())
+	{
+		if (ItemDesc.Struct != nullptr)
+		{
+			if (ItemDesc.Struct->IsChildOf(UWorldSubsystem::StaticClass()))
+			{
+				UWorldSubsystem* Subsystem = World->GetSubsystemBase(Cast<UClass>(const_cast<UStruct*>(ItemDesc.Struct)));
+				StateTreeContext.SetExternalData(ItemDesc.Handle, FStateTreeDataView(Subsystem));
+			}
+			else if (ItemDesc.Struct->IsChildOf(UActorComponent::StaticClass()))
+			{
+				UActorComponent* Component = GetOwner()->FindComponentByClass(Cast<UClass>(const_cast<UStruct*>(ItemDesc.Struct)));
+				StateTreeContext.SetExternalData(ItemDesc.Handle, FStateTreeDataView(Component));
+			}
+			else if (ItemDesc.Struct->IsChildOf(APawn::StaticClass()))
+			{
+				APawn* OwnerPawn = (AIOwner != nullptr) ? AIOwner->GetPawn() : Cast<APawn>(GetOwner());
+				StateTreeContext.SetExternalData(ItemDesc.Handle, FStateTreeDataView(OwnerPawn));
+			}
+			else if (ItemDesc.Struct->IsChildOf(AAIController::StaticClass()))
+			{
+				AAIController* OwnerController = (AIOwner != nullptr) ? AIOwner.Get() : Cast<AAIController>(GetOwner());
+				StateTreeContext.SetExternalData(ItemDesc.Handle, FStateTreeDataView(OwnerController));
+			}
+			else if (ItemDesc.Struct->IsChildOf(AActor::StaticClass()))
+			{
+				AActor* OwnerActor = (AIOwner != nullptr) ? AIOwner->GetPawn() : GetOwner();
+				StateTreeContext.SetExternalData(ItemDesc.Handle, FStateTreeDataView(OwnerActor));
+			}
+		}
+	}
+
+	return StateTreeContext.AreExternalDataViewsValid();
+}
+
+void UStateTreeComponent::TickComponent(float DeltaTime, enum ELevelTick TickType, FActorComponentTickFunction *ThisTickFunction)
+{
+	Super::TickComponent(DeltaTime, TickType, ThisTickFunction);
+
+	if (!bIsRunning || bIsPaused)
+	{
+		return;
+	}
+	
+	if (SetContextRequirements())
+	{
+		StateTreeContext.Tick(DeltaTime);
+	}
+}
+
+void UStateTreeComponent::StartLogic()
+{
+	STATETREE_LOG(Log, TEXT("%s: Start Logic"), ANSI_TO_TCHAR(__FUNCTION__));
+
+	if (SetContextRequirements())
+	{
+		StateTreeContext.Start();
+		bIsRunning = true;
+	}
+}
+
+void UStateTreeComponent::RestartLogic()
+{
+	STATETREE_LOG(Log, TEXT("%s: Restart Logic"), ANSI_TO_TCHAR(__FUNCTION__));
+
+	if (SetContextRequirements())
+	{
+		StateTreeContext.Start();
+		bIsRunning = true;
+	}
+}
+
+void UStateTreeComponent::StopLogic(const FString& Reason)
+{
+	STATETREE_LOG(Log, TEXT("%s: Stopping, reason: \'%s\'"), ANSI_TO_TCHAR(__FUNCTION__), *Reason);
+
+	if (SetContextRequirements())
+	{
+		StateTreeContext.Stop();
+		bIsRunning = false;
+	}
+}
+
+void UStateTreeComponent::Cleanup()
+{
+}
+
+void UStateTreeComponent::PauseLogic(const FString& Reason)
+{
+	STATETREE_LOG(Log, TEXT("%s: Execution updates: PAUSED (%s)"), ANSI_TO_TCHAR(__FUNCTION__), *Reason);
+	bIsPaused = true;
+}
+
+EAILogicResuming::Type UStateTreeComponent::ResumeLogic(const FString& Reason)
+{
+	STATETREE_LOG(Log, TEXT("%s: Execution updates: RESUMED (%s)"), ANSI_TO_TCHAR(__FUNCTION__), *Reason);
+
+	const EAILogicResuming::Type SuperResumeResult = Super::ResumeLogic(Reason);
+
+	if (!!bIsPaused)
+	{
+		bIsPaused = false;
+
+		if (SuperResumeResult == EAILogicResuming::Continue)
+		{
+			// Nop
+		}
+		else if (SuperResumeResult == EAILogicResuming::RestartedInstead)
+		{
+			RestartLogic();
+		}
+	}
+
+	return SuperResumeResult;
+}
+
+bool UStateTreeComponent::IsRunning() const
+{
+	return bIsRunning;
+}
+
+bool UStateTreeComponent::IsPaused() const
+{
+	return bIsPaused;
+}
+
+UGameplayTasksComponent* UStateTreeComponent::GetGameplayTasksComponent(const UGameplayTask& Task) const
+{
+	const UAITask* AITask = Cast<const UAITask>(&Task);
+	return (AITask && AITask->GetAIController()) ? AITask->GetAIController()->GetGameplayTasksComponent(Task) : Task.GetGameplayTasksComponent();
+}
+
+AActor* UStateTreeComponent::GetGameplayTaskOwner(const UGameplayTask* Task) const
+{
+	if (Task == nullptr)
+	{
+		return GetAIOwner();
+	}
+
+	const UAITask* AITask = Cast<const UAITask>(Task);
+	if (AITask)
+	{
+		return AITask->GetAIController();
+	}
+
+	const UGameplayTasksComponent* TasksComponent = Task->GetGameplayTasksComponent();
+	return TasksComponent ? TasksComponent->GetGameplayTaskOwner(Task) : nullptr;
+}
+
+AActor* UStateTreeComponent::GetGameplayTaskAvatar(const UGameplayTask* Task) const
+{
+	if (Task == nullptr)
+	{
+		return GetAIOwner() ? GetAIOwner()->GetPawn() : nullptr;
+	}
+
+	const UAITask* AITask = Cast<const UAITask>(Task);
+	if (AITask)
+	{
+		return AITask->GetAIController() ? AITask->GetAIController()->GetPawn() : nullptr;
+	}
+
+	const UGameplayTasksComponent* TasksComponent = Task->GetGameplayTasksComponent();
+	return TasksComponent ? TasksComponent->GetGameplayTaskAvatar(Task) : nullptr;
+}
+
+uint8 UStateTreeComponent::GetGameplayTaskDefaultPriority() const
+{
+	return static_cast<uint8>(EAITaskPriority::AutonomousAI);
+}
+
+void UStateTreeComponent::OnGameplayTaskInitialized(UGameplayTask& Task)
+{
+	const UAITask* AITask = Cast<const UAITask>(&Task);
+	if (AITask && (AITask->GetAIController() == nullptr))
+	{
+		// this means that the task has either been created without specifying 
+		// UAITAsk::OwnerController's value (like via BP's Construct Object node)
+		// or it has been created in C++ with inappropriate function
+		UE_LOG(LogStateTree, Error, TEXT("Missing AIController in AITask %s"), *AITask->GetName());
+	}
+}
+
+#if WITH_GAMEPLAY_DEBUGGER
+FString UStateTreeComponent::GetDebugInfoString() const
+{
+	return StateTreeContext.GetDebugInfoString();
+}
+#endif // WITH_GAMEPLAY_DEBUGGER
+
+#undef STATETREE_LOG
+#undef STATETREE_CLOG
diff --git a/Engine/Plugins/Runtime/GameplayStateTree/Source/GameplayStateTreeModule/Private/Components/StateTreeComponentSchema.cpp b/Engine/Plugins/Runtime/GameplayStateTree/Source/GameplayStateTreeModule/Private/Components/StateTreeComponentSchema.cpp
new file mode 100644
--- /dev/null
+++ b/Engine/Plugins/Runtime/GameplayStateTree/Source/GameplayStateTreeModule/Private/Components/StateTreeComponentSchema.cpp
@@ -0,0 +1,28 @@
+﻿// Copyright Epic Games, Inc. All Rights Reserved.
+
+#include "Components/StateTreeComponentSchema.h"
+#include "StateTreeConditionBase.h"
+#include "StateTreeEvaluatorBase.h"
+#include "StateTreeTaskBase.h"
+#include "GameFramework/Actor.h"
+#include "Components/ActorComponent.h"
+#include "Subsystems/WorldSubsystem.h"
+
+bool UStateTreeComponentSchema::IsStructAllowed(const UScriptStruct* InScriptStruct) const
+{
+	return InScriptStruct->IsChildOf(FStateTreeConditionCommonBase::StaticStruct())
+	|| InScriptStruct->IsChildOf(FStateTreeEvaluatorCommonBase::StaticStruct())
+	|| InScriptStruct->IsChildOf(FStateTreeTaskCommonBase::StaticStruct());
+}
+
+bool UStateTreeComponentSchema::IsClassAllowed(const UClass* InClass) const
+{
+	return IsChildOfBlueprintBase(InClass);
+}
+
+bool UStateTreeComponentSchema::IsExternalItemAllowed(const UStruct& InStruct) const
+{
+	return InStruct.IsChildOf(AActor::StaticClass())
+			|| InStruct.IsChildOf(UActorComponent::StaticClass())
+			|| InStruct.IsChildOf(UWorldSubsystem::StaticClass());
+}
diff --git a/Engine/Plugins/Runtime/GameplayStateTree/Source/GameplayStateTreeModule/Private/GameplayStateTreeModule.cpp b/Engine/Plugins/Runtime/GameplayStateTree/Source/GameplayStateTreeModule/Private/GameplayStateTreeModule.cpp
new file mode 100644
--- /dev/null
+++ b/Engine/Plugins/Runtime/GameplayStateTree/Source/GameplayStateTreeModule/Private/GameplayStateTreeModule.cpp
@@ -0,0 +1,28 @@
+﻿// Copyright Epic Games, Inc. All Rights Reserved.
+
+#include "GameplayStateTreeModule.h"
+#include "Modules/ModuleManager.h"
+
+
+//----------------------------------------------------------------------//
+// IGameplayStateTreeModule
+//----------------------------------------------------------------------//
+IGameplayStateTreeModule& IGameplayStateTreeModule::Get()
+{
+	return FModuleManager::LoadModuleChecked<IGameplayStateTreeModule>("StateTreeGameplayModule");
+}
+
+bool IGameplayStateTreeModule::IsAvailable()
+{
+	return FModuleManager::Get().IsModuleLoaded("StateTreeGameplayModule");
+}
+
+
+//----------------------------------------------------------------------//
+// FStateTreeGameplayModule
+//----------------------------------------------------------------------//
+class FStateTreeGameplayModule : public IGameplayStateTreeModule
+{
+};
+
+IMPLEMENT_MODULE(FStateTreeGameplayModule, StateTreeGameplayModule)
diff --git a/Engine/Plugins/Runtime/GameplayStateTree/Source/GameplayStateTreeModule/Public/Components/StateTreeComponent.h b/Engine/Plugins/Runtime/GameplayStateTree/Source/GameplayStateTreeModule/Public/Components/StateTreeComponent.h
new file mode 100644
--- /dev/null
+++ b/Engine/Plugins/Runtime/GameplayStateTree/Source/GameplayStateTreeModule/Public/Components/StateTreeComponent.h
@@ -0,0 +1,64 @@
+// Copyright Epic Games, Inc. All Rights Reserved.
+
+#pragma once
+
+#include "GameplayTagContainer.h"
+#include "StateTreeExecutionContext.h"
+#include "BrainComponent.h"
+#include "Tasks/AITask.h"
+#include "StateTreeComponent.generated.h"
+
+class UStateTree;
+
+UCLASS(ClassGroup = AI, HideCategories = (Activation, Collision), meta = (BlueprintSpawnableComponent))
+class GAMEPLAYSTATETREEMODULE_API UStateTreeComponent : public UBrainComponent, public IGameplayTaskOwnerInterface
+{
+	GENERATED_BODY()
+public:
+	UStateTreeComponent(const FObjectInitializer& ObjectInitializer = FObjectInitializer::Get());
+
+	// BEGIN UActorComponent overrides
+	virtual void InitializeComponent() override;
+	virtual void UninitializeComponent() override;
+	virtual void TickComponent(float DeltaTime, enum ELevelTick TickType, FActorComponentTickFunction *ThisTickFunction) override;
+	// END UActorComponent overrides
+
+	// BEGIN UBrainComponent overrides
+	virtual void StartLogic() override;
+	virtual void RestartLogic() override;
+	virtual void StopLogic(const FString& Reason)  override;
+	virtual void Cleanup() override;
+	virtual void PauseLogic(const FString& Reason) override;
+	virtual EAILogicResuming::Type ResumeLogic(const FString& Reason)  override;
+	virtual bool IsRunning() const override;
+	virtual bool IsPaused() const override;
+	// END UBrainComponent overrides
+
+	// BEGIN IGameplayTaskOwnerInterface
+	virtual UGameplayTasksComponent* GetGameplayTasksComponent(const UGameplayTask& Task) const override;
+	virtual AActor* GetGameplayTaskOwner(const UGameplayTask* Task) const override;
+	virtual AActor* GetGameplayTaskAvatar(const UGameplayTask* Task) const override;
+	virtual uint8 GetGameplayTaskDefaultPriority() const override;
+	virtual void OnGameplayTaskInitialized(UGameplayTask& Task) override;
+	// END IGameplayTaskOwnerInterface
+
+#if WITH_GAMEPLAY_DEBUGGER
+	virtual FString GetDebugInfoString() const override;
+#endif // WITH_GAMEPLAY_DEBUGGER
+
+protected:
+
+	bool SetContextRequirements();
+	
+	UPROPERTY(EditDefaultsOnly, Category = AI, meta=(RequiredAssetDataTags="Schema=StateTreeComponentSchema"))
+	UStateTree* StateTree;
+
+	UPROPERTY(Transient)
+	FStateTreeExecutionContext StateTreeContext;
+
+	/** if set, state tree execution is allowed */
+	uint8 bIsRunning : 1;
+
+	/** if set, execution requests will be postponed */
+	uint8 bIsPaused : 1;
+};
diff --git a/Engine/Plugins/Runtime/GameplayStateTree/Source/GameplayStateTreeModule/Public/Components/StateTreeComponentSchema.h b/Engine/Plugins/Runtime/GameplayStateTree/Source/GameplayStateTreeModule/Public/Components/StateTreeComponentSchema.h
new file mode 100644
--- /dev/null
+++ b/Engine/Plugins/Runtime/GameplayStateTree/Source/GameplayStateTreeModule/Public/Components/StateTreeComponentSchema.h
@@ -0,0 +1,18 @@
+﻿// Copyright Epic Games, Inc. All Rights Reserved.
+
+#pragma once
+
+#include "StateTreeSchema.h"
+#include "StateTreeComponentSchema.generated.h"
+
+
+UCLASS(BlueprintType, EditInlineNew, CollapseCategories, meta = (DisplayName = "StateTree Component"))
+class UStateTreeComponentSchema : public UStateTreeSchema
+{
+	GENERATED_BODY()
+
+protected:
+	virtual bool IsStructAllowed(const UScriptStruct* InScriptStruct) const override;
+	virtual bool IsClassAllowed(const UClass* InScriptStruct) const override;
+	virtual bool IsExternalItemAllowed(const UStruct& InStruct) const override;
+};
diff --git a/Engine/Plugins/Runtime/GameplayStateTree/Source/GameplayStateTreeModule/Public/GameplayStateTreeModule.h b/Engine/Plugins/Runtime/GameplayStateTree/Source/GameplayStateTreeModule/Public/GameplayStateTreeModule.h
new file mode 100644
--- /dev/null
+++ b/Engine/Plugins/Runtime/GameplayStateTree/Source/GameplayStateTreeModule/Public/GameplayStateTreeModule.h
@@ -0,0 +1,29 @@
+﻿// Copyright Epic Games, Inc. All Rights Reserved.
+
+#pragma  once
+
+#include "Modules/ModuleInterface.h"
+
+/**
+* The public interface to this module
+*/
+class IGameplayStateTreeModule : public IModuleInterface
+{
+
+public:
+
+	/**
+	* Singleton-like access to this module's interface.  This is just for convenience!
+	* Beware of calling this during the shutdown phase, though.  Your module might have been unloaded already.
+	*
+	* @return Returns singleton instance, loading the module on demand if needed
+	*/
+	static IGameplayStateTreeModule& Get();
+
+	/**
+	* Checks to see if this module is loaded and ready.  It is only valid to call Get() if IsAvailable() returns true.
+	*
+	* @return True if the module is loaded and ready to use
+	*/
+	static bool IsAvailable();
+};
diff --git a/Engine/Plugins/Runtime/StateTree/Source/StateTreeModule/Private/Components/StateTreeBrainComponent.cpp b/Engine/Plugins/Runtime/StateTree/Source/StateTreeModule/Private/Components/StateTreeBrainComponent.cpp
deleted file mode 100644
diff --git a/Engine/Plugins/Runtime/StateTree/Source/StateTreeModule/Public/Components/StateTreeBrainComponent.h b/Engine/Plugins/Runtime/StateTree/Source/StateTreeModule/Public/Components/StateTreeBrainComponent.h
deleted file mode 100644
diff --git a/Engine/Plugins/Runtime/StateTree/Source/StateTreeModule/Public/StateTreeTaskBase.h b/Engine/Plugins/Runtime/StateTree/Source/StateTreeModule/Public/StateTreeTaskBase.h
--- a/Engine/Plugins/Runtime/StateTree/Source/StateTreeModule/Public/StateTreeTaskBase.h
+++ b/Engine/Plugins/Runtime/StateTree/Source/StateTreeModule/Public/StateTreeTaskBase.h
@@ -57,3 +57,14 @@ struct STATETREEMODULE_API FStateTreeTaskBase : public FStateTreeNodeBase
 	virtual void AppendDebugInfoString(FString& DebugString, const FStateTreeExecutionContext& Context) const;
 #endif
 };
+
+/**
+ * Base class (namespace) for all common Tasks that are generally applicable.
+ * This allows schemas to safely include all conditions child of this struct. 
+ */
+USTRUCT(meta = (Hidden))
+struct STATETREEMODULE_API FStateTreeTaskCommonBase : public FStateTreeTaskBase
+{
+	GENERATED_BODY()
+};
+
diff --git a/Engine/Plugins/Runtime/StateTree/Source/StateTreeModule/StateTreeModule.Build.cs b/Engine/Plugins/Runtime/StateTree/Source/StateTreeModule/StateTreeModule.Build.cs
--- a/Engine/Plugins/Runtime/StateTree/Source/StateTreeModule/StateTreeModule.Build.cs
+++ b/Engine/Plugins/Runtime/StateTree/Source/StateTreeModule/StateTreeModule.Build.cs
@@ -17,7 +17,6 @@ namespace UnrealBuildTool.Rules
 					"CoreUObject",
 					"Engine",
 					"AIModule",
-					"GameplayTasks",
 					"GameplayTags",
 					"StructUtils",
 				}
diff --git a/Engine/Plugins/Runtime/StateTree/StateTree.uplugin b/Engine/Plugins/Runtime/StateTree/StateTree.uplugin
--- a/Engine/Plugins/Runtime/StateTree/StateTree.uplugin
+++ b/Engine/Plugins/Runtime/StateTree/StateTree.uplugin
@@ -3,7 +3,7 @@
 	"Version" : 1,
 	"VersionName" : "0.1",
 	"FriendlyName" : "StateTree",
-	"Description" : "StateTree for AI Behaviors",
+	"Description" : "General purpose hierarchical state machine",
 	"Category" : "Gameplay",
 	"CreatedBy" : "Epic Games, Inc.",
 	"CreatedByURL" : "http://epicgames.com",
