Removed reference to GameplayInteraction from GameplayStateTree
#rnx

diff --git a/Engine/Plugins/Runtime/GameplayStateTree/GameplayStateTree.uplugin b/Engine/Plugins/Runtime/GameplayStateTree/GameplayStateTree.uplugin
--- a/Engine/Plugins/Runtime/GameplayStateTree/GameplayStateTree.uplugin
+++ b/Engine/Plugins/Runtime/GameplayStateTree/GameplayStateTree.uplugin
@@ -26,10 +26,6 @@
 		{
 			"Name": "StateTree",
 			"Enabled": true
-		},
-		{
-			"Name": "GameplayInteractions",
-			"Enabled": true
 		}
 	]
 }
\ No newline at end of file
