Fixed GameplayStateTree module name preventing module from being properly loaded
#rnx

diff --git a/Engine/Plugins/Runtime/GameplayStateTree/Source/GameplayStateTreeModule/Private/GameplayStateTreeModule.cpp b/Engine/Plugins/Runtime/GameplayStateTree/Source/GameplayStateTreeModule/Private/GameplayStateTreeModule.cpp
--- a/Engine/Plugins/Runtime/GameplayStateTree/Source/GameplayStateTreeModule/Private/GameplayStateTreeModule.cpp
+++ b/Engine/Plugins/Runtime/GameplayStateTree/Source/GameplayStateTreeModule/Private/GameplayStateTreeModule.cpp
@@ -9,20 +9,20 @@
 //----------------------------------------------------------------------//
 IGameplayStateTreeModule& IGameplayStateTreeModule::Get()
 {
-	return FModuleManager::LoadModuleChecked<IGameplayStateTreeModule>("StateTreeGameplayModule");
+	return FModuleManager::LoadModuleChecked<IGameplayStateTreeModule>("GameplayStateTreeModule");
 }
 
 bool IGameplayStateTreeModule::IsAvailable()
 {
-	return FModuleManager::Get().IsModuleLoaded("StateTreeGameplayModule");
+	return FModuleManager::Get().IsModuleLoaded("GameplayStateTreeModule");
 }
 
 
 //----------------------------------------------------------------------//
-// FStateTreeGameplayModule
+// FGameplayStateTreeModule
 //----------------------------------------------------------------------//
-class FStateTreeGameplayModule : public IGameplayStateTreeModule
+class FGameplayStateTreeModule : public IGameplayStateTreeModule
 {
 };
 
-IMPLEMENT_MODULE(FStateTreeGameplayModule, StateTreeGameplayModule)
+IMPLEMENT_MODULE(FGameplayStateTreeModule, GameplayStateTreeModule)
