Fix non unity build errors
#test compile cpp

diff --git a/Engine/Plugins/Developer/AnimationSharing/Source/AnimationSharing/Private/AdditiveAnimationInstance.cpp b/Engine/Plugins/Developer/AnimationSharing/Source/AnimationSharing/Private/AdditiveAnimationInstance.cpp
--- a/Engine/Plugins/Developer/AnimationSharing/Source/AnimationSharing/Private/AdditiveAnimationInstance.cpp
+++ b/Engine/Plugins/Developer/AnimationSharing/Source/AnimationSharing/Private/AdditiveAnimationInstance.cpp
@@ -3,6 +3,7 @@
 #include "AdditiveAnimationInstance.h"
 #include "AnimationSharingInstances.h"
 #include "AnimationSharingManager.h"
+#include "Animation/AnimSequence.h"
 
 FAdditiveAnimationInstance::FAdditiveAnimationInstance() : SkeletalMeshComponent(nullptr), AdditiveInstance(nullptr), AdditiveAnimationSequence(nullptr), BaseComponent(nullptr), bLoopingState(false)
 {
diff --git a/Engine/Plugins/Runtime/GameplayStateTree/Source/GameplayStateTreeModule/Private/Components/StateTreeComponent.cpp b/Engine/Plugins/Runtime/GameplayStateTree/Source/GameplayStateTreeModule/Private/Components/StateTreeComponent.cpp
--- a/Engine/Plugins/Runtime/GameplayStateTree/Source/GameplayStateTreeModule/Private/Components/StateTreeComponent.cpp
+++ b/Engine/Plugins/Runtime/GameplayStateTree/Source/GameplayStateTreeModule/Private/Components/StateTreeComponent.cpp
@@ -9,6 +9,7 @@
 #include "Conditions/StateTreeCondition_Common.h"
 #include "AIController.h"
 #include "GameFramework/CharacterMovementComponent.h"
+#include "Subsystems/WorldSubsystem.h"
 
 #define STATETREE_LOG(Verbosity, Format, ...) UE_VLOG(GetOwner(), LogStateTree, Verbosity, Format, ##__VA_ARGS__)
 #define STATETREE_CLOG(Condition, Verbosity, Format, ...) UE_CVLOG((Condition), GetOwner(), LogStateTree, Verbosity, Format, ##__VA_ARGS__)
