StateTree: Made StateTreeComponent to use StateTreeRef, and fixed StateTreeRef on Actor Components

diff --git a/Engine/Plugins/Runtime/GameplayStateTree/Source/GameplayStateTreeModule/Private/Components/StateTreeComponent.cpp b/Engine/Plugins/Runtime/GameplayStateTree/Source/GameplayStateTreeModule/Private/Components/StateTreeComponent.cpp
--- a/Engine/Plugins/Runtime/GameplayStateTree/Source/GameplayStateTreeModule/Private/Components/StateTreeComponent.cpp
+++ b/Engine/Plugins/Runtime/GameplayStateTree/Source/GameplayStateTreeModule/Private/Components/StateTreeComponent.cpp
@@ -28,18 +28,34 @@ UStateTreeComponent::UStateTreeComponent(const FObjectInitializer& ObjectInitial
 
 void UStateTreeComponent::InitializeComponent()
 {
-	if (StateTree == nullptr)
+	if (StateTreeRef.StateTree == nullptr)
 	{
 		STATETREE_LOG(Error, TEXT("%s: StateTree asset is not set, cannot initialize."), ANSI_TO_TCHAR(__FUNCTION__));
 		return;
 	}
 
-	if (!StateTreeContext.Init(*GetOwner(), *StateTree, EStateTreeStorage::Internal))
+	if (!StateTreeContext.Init(*GetOwner(), *StateTreeRef.StateTree, EStateTreeStorage::Internal))
 	{
 		STATETREE_LOG(Error, TEXT("%s: Failed to init StateTreeContext."), ANSI_TO_TCHAR(__FUNCTION__));
 	}
 }
 
+#if WITH_EDITOR
+void UStateTreeComponent::PostLoad()
+{
+	Super::PostLoad();
+	
+	PRAGMA_DISABLE_DEPRECATION_WARNINGS
+	if (StateTree_DEPRECATED != nullptr)
+	{
+		StateTreeRef.StateTree = StateTree_DEPRECATED;
+		StateTreeRef.SyncParameters();
+		StateTree_DEPRECATED = nullptr;
+	}
+	PRAGMA_ENABLE_DEPRECATION_WARNINGS
+}
+#endif //WITH_EDITOR
+
 void UStateTreeComponent::UninitializeComponent()
 {
 }
@@ -113,6 +129,8 @@ void UStateTreeComponent::StartLogic()
 
 	if (SetContextRequirements())
 	{
+		StateTreeContext.SetParameters(StateTreeRef.Parameters);
+
 		StateTreeContext.Start();
 		bIsRunning = true;
 	}
@@ -124,6 +142,8 @@ void UStateTreeComponent::RestartLogic()
 
 	if (SetContextRequirements())
 	{
+		StateTreeContext.SetParameters(StateTreeRef.Parameters);
+
 		StateTreeContext.Start();
 		bIsRunning = true;
 	}
diff --git a/Engine/Plugins/Runtime/GameplayStateTree/Source/GameplayStateTreeModule/Public/Components/StateTreeComponent.h b/Engine/Plugins/Runtime/GameplayStateTree/Source/GameplayStateTreeModule/Public/Components/StateTreeComponent.h
--- a/Engine/Plugins/Runtime/GameplayStateTree/Source/GameplayStateTreeModule/Public/Components/StateTreeComponent.h
+++ b/Engine/Plugins/Runtime/GameplayStateTree/Source/GameplayStateTreeModule/Public/Components/StateTreeComponent.h
@@ -5,7 +5,9 @@
 #include "GameplayTagContainer.h"
 #include "StateTreeExecutionContext.h"
 #include "BrainComponent.h"
+#include "StateTreeReference.h"
 #include "Tasks/AITask.h"
+#include "PropertyBag.h"
 #include "StateTreeComponent.generated.h"
 
 class UStateTree;
@@ -48,10 +50,20 @@ public:
 
 protected:
 
-	bool SetContextRequirements();
+#if WITH_EDITORONLY_DATA
+	virtual void PostLoad() override;
+#endif
 	
-	UPROPERTY(EditAnywhere, BlueprintReadWrite, Category = AI, meta=(RequiredAssetDataTags="Schema=StateTreeComponentSchema"))
-	UStateTree* StateTree;
+	bool SetContextRequirements();
+
+#if WITH_EDITORONLY_DATA
+	UE_DEPRECATED(5.1, "This property has been deprecated. Use StateTreeReference instead.")
+	UPROPERTY()
+	UStateTree* StateTree_DEPRECATED;
+#endif
+
+	UPROPERTY(EditAnywhere, Category = AI, meta=(RequiredAssetDataTags = "Schema=StateTreeComponentSchema"))
+	FStateTreeReference StateTreeRef;
 
 	UPROPERTY(Transient)
 	FStateTreeExecutionContext StateTreeContext;
diff --git a/Engine/Plugins/Runtime/StateTree/Source/StateTreeModule/Private/StateTreeReference.cpp b/Engine/Plugins/Runtime/StateTree/Source/StateTreeModule/Private/StateTreeReference.cpp
--- a/Engine/Plugins/Runtime/StateTree/Source/StateTreeModule/Private/StateTreeReference.cpp
+++ b/Engine/Plugins/Runtime/StateTree/Source/StateTreeModule/Private/StateTreeReference.cpp
@@ -35,6 +35,13 @@ void FStateTreeReference::PostSerialize(const FArchive& Ar)
 {
 	if (Ar.IsLoading() && Ar.IsPersistent())
 	{
+		// Make sure the StateTree asset is fully loaded, so that the SyncParameters() will execute on valid data.
+		if (StateTree != nullptr)
+		{
+			FArchive* NonConstAr = const_cast<FArchive*>(&Ar);
+			NonConstAr->Preload(StateTree);
+		}
+
 		// This might modify the object but we don't want to dirty on load
 		SyncParameters();
 	}
diff --git a/Engine/Plugins/Runtime/StateTree/Source/StateTreeModule/Public/StateTreeReference.h b/Engine/Plugins/Runtime/StateTree/Source/StateTreeModule/Public/StateTreeReference.h
--- a/Engine/Plugins/Runtime/StateTree/Source/StateTreeModule/Public/StateTreeReference.h
+++ b/Engine/Plugins/Runtime/StateTree/Source/StateTreeModule/Public/StateTreeReference.h
@@ -18,10 +18,10 @@ struct STATETREEMODULE_API FStateTreeReference
 	FStateTreeReference();
 	~FStateTreeReference();
 	
-	UPROPERTY(EditDefaultsOnly, Category="")
+	UPROPERTY(EditAnywhere, Category = "")
 	TObjectPtr<UStateTree> StateTree = nullptr;
 
-	UPROPERTY(EditDefaultsOnly, Category="", meta=(FixedLayout))
+	UPROPERTY(EditAnywhere, Category = "", meta = (FixedLayout))
 	FInstancedPropertyBag Parameters;
 
 #if WITH_EDITOR
