StateTree: Fix missing include and editor guard from StateTreeComponentSchema

diff --git a/Engine/Plugins/Runtime/GameplayStateTree/Source/GameplayStateTreeModule/Private/Components/StateTreeComponentSchema.cpp b/Engine/Plugins/Runtime/GameplayStateTree/Source/GameplayStateTreeModule/Private/Components/StateTreeComponentSchema.cpp
--- a/Engine/Plugins/Runtime/GameplayStateTree/Source/GameplayStateTreeModule/Private/Components/StateTreeComponentSchema.cpp
+++ b/Engine/Plugins/Runtime/GameplayStateTree/Source/GameplayStateTreeModule/Private/Components/StateTreeComponentSchema.cpp
@@ -39,6 +39,7 @@ TConstArrayView<FStateTreeExternalDataDesc> UStateTreeComponentSchema::GetNamedE
 	return MakeArrayView(&ContextActorDataDesc, 1);
 }
 
+#if WITH_EDITOR
 void UStateTreeComponentSchema::PostEditChangeChainProperty(FPropertyChangedChainEvent& PropertyChangedEvent)
 {
 	Super::PostEditChangeChainProperty(PropertyChangedEvent);
@@ -54,3 +55,4 @@ void UStateTreeComponentSchema::PostEditChangeChainProperty(FPropertyChangedChai
 		}
 	}
 }
+#endif
\ No newline at end of file
diff --git a/Engine/Plugins/Runtime/GameplayStateTree/Source/GameplayStateTreeModule/Public/Components/StateTreeComponentSchema.h b/Engine/Plugins/Runtime/GameplayStateTree/Source/GameplayStateTreeModule/Public/Components/StateTreeComponentSchema.h
--- a/Engine/Plugins/Runtime/GameplayStateTree/Source/GameplayStateTreeModule/Public/Components/StateTreeComponentSchema.h
+++ b/Engine/Plugins/Runtime/GameplayStateTree/Source/GameplayStateTreeModule/Public/Components/StateTreeComponentSchema.h
@@ -3,6 +3,7 @@
 #pragma once
 
 #include "StateTreeSchema.h"
+#include "Templates/SubclassOf.h"
 #include "StateTreeComponentSchema.generated.h"
 
 /**
