StateTreeSchema: fix missing Actor include

diff --git a/Engine/Plugins/Runtime/GameplayStateTree/Source/GameplayStateTreeModule/Public/Components/StateTreeComponentSchema.h b/Engine/Plugins/Runtime/GameplayStateTree/Source/GameplayStateTreeModule/Public/Components/StateTreeComponentSchema.h
--- a/Engine/Plugins/Runtime/GameplayStateTree/Source/GameplayStateTreeModule/Public/Components/StateTreeComponentSchema.h
+++ b/Engine/Plugins/Runtime/GameplayStateTree/Source/GameplayStateTreeModule/Public/Components/StateTreeComponentSchema.h
@@ -3,6 +3,7 @@
 #pragma once
 
 #include "StateTreeSchema.h"
+#include "GameFramework/Actor.h"
 #include "Templates/SubclassOf.h"
 #include "StateTreeComponentSchema.generated.h"
 
