PR #9910: Create Setter for auto start property for state tree (Contributed by [REDACTED_PERSON])

diff --git a/Engine/Plugins/Runtime/GameplayStateTree/Source/GameplayStateTreeModule/Private/Components/StateTreeComponent.cpp b/Engine/Plugins/Runtime/GameplayStateTree/Source/GameplayStateTreeModule/Private/Components/StateTreeComponent.cpp
--- a/Engine/Plugins/Runtime/GameplayStateTree/Source/GameplayStateTreeModule/Private/Components/StateTreeComponent.cpp
+++ b/Engine/Plugins/Runtime/GameplayStateTree/Source/GameplayStateTreeModule/Private/Components/StateTreeComponent.cpp
@@ -322,6 +322,11 @@ void UStateTreeComponent::OnGameplayTaskInitialized(UGameplayTask& Task)
 	}
 }
 
+void UStateTreeComponent::SetStartLogicAutomatically(const bool bInStartLogicAutomatically)
+{
+	bStartLogicAutomatically = bInStartLogicAutomatically;
+}
+
 void UStateTreeComponent::SendStateTreeEvent(const FStateTreeEvent& Event)
 {
 	SendStateTreeEvent(Event.Tag, Event.Payload, Event.Origin);
diff --git a/Engine/Plugins/Runtime/GameplayStateTree/Source/GameplayStateTreeModule/Public/Components/StateTreeComponent.h b/Engine/Plugins/Runtime/GameplayStateTree/Source/GameplayStateTreeModule/Public/Components/StateTreeComponent.h
--- a/Engine/Plugins/Runtime/GameplayStateTree/Source/GameplayStateTreeModule/Public/Components/StateTreeComponent.h
+++ b/Engine/Plugins/Runtime/GameplayStateTree/Source/GameplayStateTreeModule/Public/Components/StateTreeComponent.h
@@ -45,6 +45,10 @@ public:
 	virtual void OnGameplayTaskInitialized(UGameplayTask& Task) override;
 	// END IGameplayTaskOwnerInterface
 
+	/** Sets whether the State Tree is started automatically on being play. */
+	UFUNCTION(BlueprintCallable, Category = "Gameplay|StateTree")
+	void SetStartLogicAutomatically(const bool bInStartLogicAutomatically);
+
 	/** Sends event to the running StateTree. */
 	UFUNCTION(BlueprintCallable, Category = "Gameplay|StateTree")
 	void SendStateTreeEvent(const FStateTreeEvent& Event);
@@ -73,7 +77,7 @@ protected:
 	UPROPERTY(EditAnywhere, Category = AI, meta=(Schema="/Script/GameplayStateTreeModule.StateTreeComponentSchema"))
 	FStateTreeReference StateTreeRef;
 
-	/** If true, the StateTree logic is started on being play. Otherwise StartLogic() needs to be called. */
+	/** If true, the StateTree logic is started on begin play. Otherwise StartLogic() needs to be called. */
 	UPROPERTY(EditAnywhere, Category = AI)
 	bool bStartLogicAutomatically = true;
 
