PR #9880: [StateTree] Fix crash in UStateTreeComponent (Contributed by [REDACTED_PERSON])

diff --git a/Engine/Plugins/Runtime/GameplayStateTree/Source/GameplayStateTreeModule/Private/Components/StateTreeComponent.cpp b/Engine/Plugins/Runtime/GameplayStateTree/Source/GameplayStateTreeModule/Private/Components/StateTreeComponent.cpp
--- a/Engine/Plugins/Runtime/GameplayStateTree/Source/GameplayStateTreeModule/Private/Components/StateTreeComponent.cpp
+++ b/Engine/Plugins/Runtime/GameplayStateTree/Source/GameplayStateTreeModule/Private/Components/StateTreeComponent.cpp
@@ -176,6 +176,12 @@ void UStateTreeComponent::TickComponent(float DeltaTime, enum ELevelTick TickTyp
 	{
 		return;
 	}
+	
+	if (!StateTreeRef.IsValid())
+	{
+		STATETREE_LOG(Warning, TEXT("%s: Trying to tick State Tree component with invalid asset."), ANSI_TO_TCHAR(__FUNCTION__));
+		return;
+	}
 
 	FStateTreeExecutionContext Context(*GetOwner(), *StateTreeRef.GetStateTree(), InstanceData);
 	if (SetContextRequirements(Context))
@@ -188,6 +194,12 @@ void UStateTreeComponent::StartLogic()
 {
 	STATETREE_LOG(Log, TEXT("%s: Start Logic"), ANSI_TO_TCHAR(__FUNCTION__));
 
+	if (!StateTreeRef.IsValid())
+	{
+		STATETREE_LOG(Warning, TEXT("%s: Trying to start State Tree component with invalid asset."), ANSI_TO_TCHAR(__FUNCTION__));
+		return;
+	}
+
 	FStateTreeExecutionContext Context(*GetOwner(), *StateTreeRef.GetStateTree(), InstanceData);
 	if (SetContextRequirements(Context))
 	{
@@ -201,6 +213,12 @@ void UStateTreeComponent::RestartLogic()
 {
 	STATETREE_LOG(Log, TEXT("%s: Restart Logic"), ANSI_TO_TCHAR(__FUNCTION__));
 
+	if (!StateTreeRef.IsValid())
+	{
+		STATETREE_LOG(Warning, TEXT("%s: Trying to restart State Tree component with invalid asset."), ANSI_TO_TCHAR(__FUNCTION__));
+		return;
+	}
+
 	FStateTreeExecutionContext Context(*GetOwner(), *StateTreeRef.GetStateTree(), InstanceData);
 	if (SetContextRequirements(Context))
 	{
@@ -214,6 +232,17 @@ void UStateTreeComponent::StopLogic(const FString& Reason)
 {
 	STATETREE_LOG(Log, TEXT("%s: Stopping, reason: \'%s\'"), ANSI_TO_TCHAR(__FUNCTION__), *Reason);
 
+	if (!bIsRunning)
+	{
+		return;
+	}
+
+	if (!StateTreeRef.IsValid())
+	{
+		STATETREE_LOG(Warning, TEXT("%s: Trying to stop State Tree component with invalid asset."), ANSI_TO_TCHAR(__FUNCTION__));
+		return;
+	}
+
 	FStateTreeExecutionContext Context(*GetOwner(), *StateTreeRef.GetStateTree(), InstanceData);
 	if (SetContextRequirements(Context))
 	{
@@ -336,7 +365,13 @@ void UStateTreeComponent::SendStateTreeEvent(const FGameplayTag Tag, const FCons
 {
 	if (!bIsRunning)
 	{
-		STATETREE_LOG(Warning, TEXT("%s: Trying to send even to a StateTree that is not started yet."), ANSI_TO_TCHAR(__FUNCTION__));
+		STATETREE_LOG(Warning, TEXT("%s: Trying to send event to a State Tree that is not started yet."), ANSI_TO_TCHAR(__FUNCTION__));
+		return;
+	}
+
+	if (!StateTreeRef.IsValid())
+	{
+		STATETREE_LOG(Warning, TEXT("%s: Trying to send event to State Tree component with invalid asset."), ANSI_TO_TCHAR(__FUNCTION__));
 		return;
 	}
 
