PR #9830: [StateTree] StateTreeComponent.GetStateTreeRunStatus method and status delegate (Contributed by [REDACTED_PERSON])

diff --git a/Engine/Plugins/Runtime/GameplayStateTree/Source/GameplayStateTreeModule/Private/Components/StateTreeComponent.cpp b/Engine/Plugins/Runtime/GameplayStateTree/Source/GameplayStateTreeModule/Private/Components/StateTreeComponent.cpp
--- a/Engine/Plugins/Runtime/GameplayStateTree/Source/GameplayStateTreeModule/Private/Components/StateTreeComponent.cpp
+++ b/Engine/Plugins/Runtime/GameplayStateTree/Source/GameplayStateTreeModule/Private/Components/StateTreeComponent.cpp
@@ -186,7 +186,13 @@ void UStateTreeComponent::TickComponent(float DeltaTime, enum ELevelTick TickTyp
 	FStateTreeExecutionContext Context(*GetOwner(), *StateTreeRef.GetStateTree(), InstanceData);
 	if (SetContextRequirements(Context))
 	{
-		Context.Tick(DeltaTime);
+		const EStateTreeRunStatus PreviousRunStatus = Context.GetStateTreeRunStatus();
+		const EStateTreeRunStatus CurrentRunStatus = Context.Tick(DeltaTime);
+
+		if (CurrentRunStatus != PreviousRunStatus)
+		{
+			OnStateTreeRunStatusChanged.Broadcast(CurrentRunStatus);
+		}
 	}
 }
 
@@ -204,8 +210,14 @@ void UStateTreeComponent::StartLogic()
 	if (SetContextRequirements(Context))
 	{
 		Context.SetParameters(StateTreeRef.GetParameters());
-		Context.Start();
+		const EStateTreeRunStatus PreviousRunStatus = Context.GetStateTreeRunStatus();
+		const EStateTreeRunStatus CurrentRunStatus = Context.Start();
 		bIsRunning = true;
+		
+		if (CurrentRunStatus != PreviousRunStatus)
+		{
+			OnStateTreeRunStatusChanged.Broadcast(CurrentRunStatus);
+		}
 	}
 }
 
@@ -223,8 +235,14 @@ void UStateTreeComponent::RestartLogic()
 	if (SetContextRequirements(Context))
 	{
 		Context.SetParameters(StateTreeRef.GetParameters());
-		Context.Start();
+		const EStateTreeRunStatus PreviousRunStatus = Context.GetStateTreeRunStatus();
+		const EStateTreeRunStatus CurrentRunStatus = Context.Start();
 		bIsRunning = true;
+		
+		if (CurrentRunStatus != PreviousRunStatus)
+		{
+			OnStateTreeRunStatusChanged.Broadcast(CurrentRunStatus);
+		}
 	}
 }
 
@@ -246,8 +264,14 @@ void UStateTreeComponent::StopLogic(const FString& Reason)
 	FStateTreeExecutionContext Context(*GetOwner(), *StateTreeRef.GetStateTree(), InstanceData);
 	if (SetContextRequirements(Context))
 	{
-		Context.Stop();
+		const EStateTreeRunStatus PreviousRunStatus = Context.GetStateTreeRunStatus();
+		const EStateTreeRunStatus CurrentRunStatus = Context.Stop();
 		bIsRunning = false;
+
+		if (CurrentRunStatus != PreviousRunStatus)
+		{
+			OnStateTreeRunStatusChanged.Broadcast(CurrentRunStatus);
+		}
 	}
 }
 
@@ -382,6 +406,16 @@ void UStateTreeComponent::SendStateTreeEvent(const FGameplayTag Tag, const FCons
 	}
 }
 
+EStateTreeRunStatus UStateTreeComponent::GetStateTreeRunStatus() const
+{
+	FStateTreeExecutionContext Context(*GetOwner(), *StateTreeRef.GetStateTree(), const_cast<FStateTreeInstanceData&>(InstanceData));
+	if (Context.IsValid())
+	{
+		return Context.GetStateTreeRunStatus();
+	}
+
+	return EStateTreeRunStatus::Unset;
+}
 
 #if WITH_GAMEPLAY_DEBUGGER
 FString UStateTreeComponent::GetDebugInfoString() const
diff --git a/Engine/Plugins/Runtime/GameplayStateTree/Source/GameplayStateTreeModule/Public/Components/StateTreeComponent.h b/Engine/Plugins/Runtime/GameplayStateTree/Source/GameplayStateTreeModule/Public/Components/StateTreeComponent.h
--- a/Engine/Plugins/Runtime/GameplayStateTree/Source/GameplayStateTreeModule/Public/Components/StateTreeComponent.h
+++ b/Engine/Plugins/Runtime/GameplayStateTree/Source/GameplayStateTreeModule/Public/Components/StateTreeComponent.h
@@ -12,6 +12,8 @@
 
 class UStateTree;
 
+DECLARE_DYNAMIC_MULTICAST_DELEGATE_OneParam(FStateTreeRunStatusChanged, EStateTreeRunStatus, StateTreeRunStatus);
+
 UCLASS(ClassGroup = AI, HideCategories = (Activation, Collision), meta = (BlueprintSpawnableComponent))
 class GAMEPLAYSTATETREEMODULE_API UStateTreeComponent : public UBrainComponent, public IGameplayTaskOwnerInterface
 {
@@ -56,6 +58,14 @@ public:
 	/** Sends event to the running StateTree. */
 	void SendStateTreeEvent(const FGameplayTag Tag, const FConstStructView Payload = FConstStructView(), const FName Origin = FName());
 
+	/** Returns the current run status of the StateTree. */
+	UFUNCTION(BlueprintPure, Category = "Gameplay|StateTree")
+	EStateTreeRunStatus GetStateTreeRunStatus() const;
+
+	/** Called when the run status of the StateTree has changed */
+	UPROPERTY(BlueprintAssignable, Category = "Gameplay|StateTree")
+	FStateTreeRunStatusChanged OnStateTreeRunStatusChanged;
+
 #if WITH_GAMEPLAY_DEBUGGER
 	virtual FString GetDebugInfoString() const override;
 #endif // WITH_GAMEPLAY_DEBUGGER
diff --git a/Engine/Plugins/Runtime/StateTree/Source/StateTreeModule/Private/StateTreeExecutionContext.cpp b/Engine/Plugins/Runtime/StateTree/Source/StateTreeModule/Private/StateTreeExecutionContext.cpp
--- a/Engine/Plugins/Runtime/StateTree/Source/StateTreeModule/Private/StateTreeExecutionContext.cpp
+++ b/Engine/Plugins/Runtime/StateTree/Source/StateTreeModule/Private/StateTreeExecutionContext.cpp
@@ -429,6 +429,26 @@ EStateTreeRunStatus FStateTreeExecutionContext::Tick(const float DeltaTime)
 	return Exec->TreeRunStatus;
 }
 
+EStateTreeRunStatus FStateTreeExecutionContext::GetStateTreeRunStatus() const
+{
+	if (!IsValid())
+	{
+		STATETREE_LOG(Warning, TEXT("%s: StateTree context is not initialized properly ('%s' using StateTree '%s')"),
+			ANSI_TO_TCHAR(__FUNCTION__), *GetNameSafe(&Owner), *GetFullNameSafe(&StateTree));
+		return EStateTreeRunStatus::Failed;
+	}
+
+	if (!InstanceData.IsValid())
+	{
+		STATETREE_LOG(Error, TEXT("%s: GetStateTreeRunStatus called on %s using StateTree %s with invalid instance data. Start() must be called before GetStateTreeRunStatus()."),
+			ANSI_TO_TCHAR(__FUNCTION__), *GetNameSafe(&Owner), *GetFullNameSafe(&StateTree));
+		return EStateTreeRunStatus::Failed;
+	}
+
+	const FStateTreeExecutionState& Exec = GetExecState();
+	return Exec.TreeRunStatus;
+}
+
 void FStateTreeExecutionContext::SendEvent(const FStateTreeEvent& Event) const
 {
 	SendEvent(Event.Tag, Event.Payload, Event.Origin);
diff --git a/Engine/Plugins/Runtime/StateTree/Source/StateTreeModule/Public/StateTreeExecutionContext.h b/Engine/Plugins/Runtime/StateTree/Source/StateTreeModule/Public/StateTreeExecutionContext.h
--- a/Engine/Plugins/Runtime/StateTree/Source/StateTreeModule/Public/StateTreeExecutionContext.h
+++ b/Engine/Plugins/Runtime/StateTree/Source/StateTreeModule/Public/StateTreeExecutionContext.h
@@ -129,9 +129,16 @@ public:
 	/** Stop executing. */
 	EStateTreeRunStatus Stop();
 
-	/** Tick the state tree logic. */
+	/**
+	 * Tick the state tree logic.
+	 * @param DeltaTime time to advance the logic.
+	 * @returns tree run status after the tick.
+	 */
 	EStateTreeRunStatus Tick(const float DeltaTime);
 
+	/** @return the tree run status. */
+	EStateTreeRunStatus GetStateTreeRunStatus() const;
+
 	/** @return the status of the last tick function */
 	EStateTreeRunStatus GetLastTickStatus() const;
 
diff --git a/Engine/Plugins/Runtime/StateTree/Source/StateTreeModule/Public/StateTreeTypes.h b/Engine/Plugins/Runtime/StateTree/Source/StateTreeModule/Public/StateTreeTypes.h
--- a/Engine/Plugins/Runtime/StateTree/Source/StateTreeModule/Public/StateTreeTypes.h
+++ b/Engine/Plugins/Runtime/StateTree/Source/StateTreeModule/Public/StateTreeTypes.h
@@ -13,8 +13,8 @@ STATETREEMODULE_API DECLARE_LOG_CATEGORY_EXTERN(LogStateTree, Warning, All);
 #define WITH_STATETREE_DEBUG (!(UE_BUILD_SHIPPING || UE_BUILD_SHIPPING_WITH_EDITOR || UE_BUILD_TEST) && 1)
 #endif // WITH_STATETREE_DEBUG
 
-/** Status describing current ticking state. */
-UENUM()
+/** Status describing current run status of a State Tree. */
+UENUM(BlueprintType)
 enum class EStateTreeRunStatus : uint8
 {
 	/** Tree is still running. */
