StateTree: Set parameters when calling Tick() and Stop()
- Fixed bug where parameter data was not set on tick or stop

diff --git a/Engine/Plugins/Runtime/GameplayStateTree/Source/GameplayStateTreeModule/Private/Components/StateTreeComponent.cpp b/Engine/Plugins/Runtime/GameplayStateTree/Source/GameplayStateTreeModule/Private/Components/StateTreeComponent.cpp
--- a/Engine/Plugins/Runtime/GameplayStateTree/Source/GameplayStateTreeModule/Private/Components/StateTreeComponent.cpp
+++ b/Engine/Plugins/Runtime/GameplayStateTree/Source/GameplayStateTreeModule/Private/Components/StateTreeComponent.cpp
@@ -182,6 +182,7 @@ void UStateTreeComponent::TickComponent(float DeltaTime, enum ELevelTick TickTyp
 	FStateTreeExecutionContext Context(*GetOwner(), *StateTreeRef.GetStateTree(), InstanceData);
 	if (SetContextRequirements(Context))
 	{
+		Context.SetParameters(StateTreeRef.GetParameters());
 		const EStateTreeRunStatus PreviousRunStatus = Context.GetStateTreeRunStatus();
 		const EStateTreeRunStatus CurrentRunStatus = Context.Tick(DeltaTime);
 
@@ -260,6 +261,7 @@ void UStateTreeComponent::StopLogic(const FString& Reason)
 	FStateTreeExecutionContext Context(*GetOwner(), *StateTreeRef.GetStateTree(), InstanceData);
 	if (SetContextRequirements(Context))
 	{
+		Context.SetParameters(StateTreeRef.GetParameters());
 		const EStateTreeRunStatus PreviousRunStatus = Context.GetStateTreeRunStatus();
 		const EStateTreeRunStatus CurrentRunStatus = Context.Stop();
 		bIsRunning = false;
