StateTree: Clarify usage of SetStartLogicAutomatically

diff --git a/Engine/Plugins/Runtime/GameplayStateTree/Source/GameplayStateTreeModule/Public/Components/StateTreeComponent.h b/Engine/Plugins/Runtime/GameplayStateTree/Source/GameplayStateTreeModule/Public/Components/StateTreeComponent.h
--- a/Engine/Plugins/Runtime/GameplayStateTree/Source/GameplayStateTreeModule/Public/Components/StateTreeComponent.h
+++ b/Engine/Plugins/Runtime/GameplayStateTree/Source/GameplayStateTreeModule/Public/Components/StateTreeComponent.h
@@ -51,7 +51,11 @@ public:
 	virtual void OnGameplayTaskInitialized(UGameplayTask& Task) override;
 	// END IGameplayTaskOwnerInterface
 
-	/** Sets whether the State Tree is started automatically on being play. */
+	/**
+	 * Sets whether the State Tree is started automatically on being play.
+	 * This function sets the bStartLogicAutomatically property, and should be used mostly from constructions sscripts.
+	 * If you wish to start the logic manually, call StartLogic(). 
+	 */
 	UFUNCTION(BlueprintCallable, Category = "Gameplay|StateTree")
 	void SetStartLogicAutomatically(const bool bInStartLogicAutomatically);
 
