[StateTree] stop logic on EndPlay for auto-activated components (start logic) that are not associated to an AI owner.
#rnx

diff --git a/Engine/Plugins/Runtime/GameplayStateTree/Source/GameplayStateTreeModule/Private/Components/StateTreeComponent.cpp b/Engine/Plugins/Runtime/GameplayStateTree/Source/GameplayStateTreeModule/Private/Components/StateTreeComponent.cpp
--- a/Engine/Plugins/Runtime/GameplayStateTree/Source/GameplayStateTreeModule/Private/Components/StateTreeComponent.cpp
+++ b/Engine/Plugins/Runtime/GameplayStateTree/Source/GameplayStateTreeModule/Private/Components/StateTreeComponent.cpp
@@ -164,6 +164,13 @@ void UStateTreeComponent::BeginPlay()
 	}
 }
 
+void UStateTreeComponent::EndPlay(const EEndPlayReason::Type EndPlayReason)
+{
+	StopLogic(UEnum::GetValueAsString(EndPlayReason));
+
+	Super::EndPlay(EndPlayReason);
+}
+
 void UStateTreeComponent::TickComponent(float DeltaTime, enum ELevelTick TickType, FActorComponentTickFunction *ThisTickFunction)
 {
 	Super::TickComponent(DeltaTime, TickType, ThisTickFunction);
@@ -275,6 +282,7 @@ void UStateTreeComponent::StopLogic(const FString& Reason)
 
 void UStateTreeComponent::Cleanup()
 {
+	StopLogic(TEXT("Cleanup"));
 }
 
 void UStateTreeComponent::PauseLogic(const FString& Reason)
diff --git a/Engine/Plugins/Runtime/GameplayStateTree/Source/GameplayStateTreeModule/Public/Components/StateTreeComponent.h b/Engine/Plugins/Runtime/GameplayStateTree/Source/GameplayStateTreeModule/Public/Components/StateTreeComponent.h
--- a/Engine/Plugins/Runtime/GameplayStateTree/Source/GameplayStateTreeModule/Public/Components/StateTreeComponent.h
+++ b/Engine/Plugins/Runtime/GameplayStateTree/Source/GameplayStateTreeModule/Public/Components/StateTreeComponent.h
@@ -29,6 +29,7 @@ public:
 	virtual void InitializeComponent() override;
 	virtual void UninitializeComponent() override;
 	virtual void BeginPlay() override;
+	virtual void EndPlay(const EEndPlayReason::Type EndPlayReason) override;
 	virtual void TickComponent(float DeltaTime, enum ELevelTick TickType, FActorComponentTickFunction *ThisTickFunction) override;
 	// END UActorComponent overrides
 
