PR #10100: Including GAMEPLAYSTATETREEMODULE_API for UStateTreeComponentSchema (Contributed by [REDACTED_PERSON])

diff --git a/Engine/Plugins/Runtime/GameplayStateTree/Source/GameplayStateTreeModule/Public/Components/StateTreeComponentSchema.h b/Engine/Plugins/Runtime/GameplayStateTree/Source/GameplayStateTreeModule/Public/Components/StateTreeComponentSchema.h
--- a/Engine/Plugins/Runtime/GameplayStateTree/Source/GameplayStateTreeModule/Public/Components/StateTreeComponentSchema.h
+++ b/Engine/Plugins/Runtime/GameplayStateTree/Source/GameplayStateTreeModule/Public/Components/StateTreeComponentSchema.h
@@ -11,7 +11,7 @@
  * StateTree for Actors with StateTree component. 
  */
 UCLASS(BlueprintType, EditInlineNew, CollapseCategories, meta = (DisplayName = "StateTree Component", CommonSchema))
-class UStateTreeComponentSchema : public UStateTreeSchema
+class GAMEPLAYSTATETREEMODULE_API UStateTreeComponentSchema : public UStateTreeSchema
 {
 	GENERATED_BODY()
 
