[StateTree]
PR #11008: Fix nullptr error for invalid StateTreeRef

diff --git a/Engine/Plugins/Runtime/GameplayStateTree/Source/GameplayStateTreeModule/Private/Components/StateTreeComponent.cpp b/Engine/Plugins/Runtime/GameplayStateTree/Source/GameplayStateTreeModule/Private/Components/StateTreeComponent.cpp
--- a/Engine/Plugins/Runtime/GameplayStateTree/Source/GameplayStateTreeModule/Private/Components/StateTreeComponent.cpp
+++ b/Engine/Plugins/Runtime/GameplayStateTree/Source/GameplayStateTreeModule/Private/Components/StateTreeComponent.cpp
@@ -420,8 +420,12 @@ EStateTreeRunStatus UStateTreeComponent::GetStateTreeRunStatus() const
 #if WITH_GAMEPLAY_DEBUGGER
 FString UStateTreeComponent::GetDebugInfoString() const
 {
-	FStateTreeExecutionContext Context(*GetOwner(), *StateTreeRef.GetStateTree(), const_cast<FStateTreeInstanceData&>(InstanceData));
-	return Context.GetDebugInfoString();
+	if (!StateTreeRef.IsValid())
+	{
+		return FString("No StateTree to run.");
+	}
+
+	return FStateTreeExecutionContext(*GetOwner(), *StateTreeRef.GetStateTree(), const_cast<FStateTreeInstanceData&>(InstanceData)).GetDebugInfoString();
 }
 #endif // WITH_GAMEPLAY_DEBUGGER
 
