Create BTTask_RunStateTree and BTTask_RunDynamicStateTree that allows the execution of state tree in a behavior tree

diff --git a/Engine/Plugins/Runtime/GameplayStateTree/Source/GameplayStateTreeModule/GameplayStateTreeModule.Build.cs b/Engine/Plugins/Runtime/GameplayStateTree/Source/GameplayStateTreeModule/GameplayStateTreeModule.Build.cs
--- a/Engine/Plugins/Runtime/GameplayStateTree/Source/GameplayStateTreeModule/GameplayStateTreeModule.Build.cs
+++ b/Engine/Plugins/Runtime/GameplayStateTree/Source/GameplayStateTreeModule/GameplayStateTreeModule.Build.cs
@@ -17,6 +17,7 @@ namespace UnrealBuildTool.Rules
 					"Core",
 					"CoreUObject",
 					"Engine",
+					"GameplayTags",
 					"GameplayTasks",
 					"StateTreeModule",
 					"StructUtils"
diff --git a/Engine/Plugins/Runtime/GameplayStateTree/Source/GameplayStateTreeModule/Private/BehaviorTree/GameplayStateTreeBTUtils.cpp b/Engine/Plugins/Runtime/GameplayStateTree/Source/GameplayStateTreeModule/Private/BehaviorTree/GameplayStateTreeBTUtils.cpp
new file mode 100644
--- /dev/null
+++ b/Engine/Plugins/Runtime/GameplayStateTree/Source/GameplayStateTreeModule/Private/BehaviorTree/GameplayStateTreeBTUtils.cpp
@@ -0,0 +1,24 @@
+// Copyright Epic Games, Inc. All Rights Reserved.
+
+#include "BehaviorTree/GameplayStateTreeBTUtils.h"
+
+#include "BehaviorTree/BehaviorTreeTypes.h"
+#include "StateTreeExecutionTypes.h"
+
+namespace GameplayStateTreeBTUtils
+{
+	EBTNodeResult::Type StateTreeRunStatusToBTNodeResult(EStateTreeRunStatus Result)
+	{
+		switch (Result)
+		{
+		case EStateTreeRunStatus::Succeeded:
+			return EBTNodeResult::Succeeded;
+		case EStateTreeRunStatus::Running:
+			return EBTNodeResult::InProgress;
+		case EStateTreeRunStatus::Failed:
+		case EStateTreeRunStatus::Stopped:
+		default:
+			return EBTNodeResult::Failed;
+		}
+	}
+}
\ No newline at end of file
diff --git a/Engine/Plugins/Runtime/GameplayStateTree/Source/GameplayStateTreeModule/Private/BehaviorTree/Tasks/BTTask_RunDynamicStateTree.cpp b/Engine/Plugins/Runtime/GameplayStateTree/Source/GameplayStateTreeModule/Private/BehaviorTree/Tasks/BTTask_RunDynamicStateTree.cpp
new file mode 100644
--- /dev/null
+++ b/Engine/Plugins/Runtime/GameplayStateTree/Source/GameplayStateTreeModule/Private/BehaviorTree/Tasks/BTTask_RunDynamicStateTree.cpp
@@ -0,0 +1,134 @@
+// Copyright Epic Games, Inc. All Rights Reserved.
+
+#include "BehaviorTree/Tasks/BTTask_RunDynamicStateTree.h"
+
+#include "BehaviorTree/BehaviorTreeComponent.h"
+#include "BehaviorTree/BTCompositeNode.h"
+#include "BehaviorTree/GameplayStateTreeBTUtils.h"
+#include "GameFramework/Actor.h"
+#include "StateTreeExecutionContext.h"
+#include "VisualLogger/VisualLogger.h"
+
+bool UBTTask_RunDynamicStateTree::SetDynamicStateTree(UBehaviorTreeComponent& OwnerComp, FGameplayTag InInjectTag, const FStateTreeReference& InStateTree, const FSetContextDataDelegate& InSetContextDataDelegate, float InInterval, float InRandomDeviation, UBTCompositeNode* OptionalStartNode /*= nullptr*/)
+{
+	bool bInjected = false;
+	auto ReplaceDynamicStateTree = [&](UBTTaskNode& TaskNode, const FBehaviorTreeInstance& InstanceInfo, int32 InstanceIndex)
+		{
+			UBTTask_RunDynamicStateTree* StateTreeTask = Cast<UBTTask_RunDynamicStateTree>(&TaskNode);
+			if (StateTreeTask && StateTreeTask->InjectionTag == InInjectTag)
+			{
+				const uint8* NodeMemory = StateTreeTask->GetNodeMemory<uint8>(InstanceInfo);
+				if (UBTTask_RunDynamicStateTree* InstancedNode = Cast<UBTTask_RunDynamicStateTree>(StateTreeTask->GetNodeInstance(OwnerComp, (uint8*)NodeMemory)))
+				{
+					InstancedNode->SetStateTreeToRun(OwnerComp, InStateTree, InSetContextDataDelegate, InInterval, InRandomDeviation);
+					UE_VLOG(OwnerComp.GetOwner(), LogBehaviorTree, Log, TEXT("Replaced state tree in %s with %s (tag: %s)"),
+						*UBehaviorTreeTypes::DescribeNodeHelper(StateTreeTask), *GetNameSafe(InStateTree.GetStateTree()), *InInjectTag.ToString());
+					bInjected = true;
+
+					if (InstanceInfo.ActiveNodeType == EBTActiveNode::ActiveTask && StateTreeTask == InstanceInfo.ActiveNode)
+					{
+						UBTCompositeNode* RestartNode = StateTreeTask->GetParentNode();
+						int32 RestartChildIdx = RestartNode->GetChildIndex(*StateTreeTask);
+						OwnerComp.RequestExecution(RestartNode, InstanceIndex, StateTreeTask, RestartChildIdx, EBTNodeResult::Aborted);
+					}
+				}
+			}
+		};
+
+	if (OptionalStartNode)
+	{
+		OwnerComp.ForEachChildTask(*OptionalStartNode, OwnerComp.FindInstanceContainingNode(OptionalStartNode), ReplaceDynamicStateTree);
+	}
+	else
+	{
+		OwnerComp.ForEachChildTask(ReplaceDynamicStateTree);
+	}
+
+	return bInjected;
+}
+
+
+UBTTask_RunDynamicStateTree::UBTTask_RunDynamicStateTree(const FObjectInitializer& ObjectInitializer)
+	: Super(ObjectInitializer)
+{
+	INIT_TASK_NODE_NOTIFY_FLAGS();
+	NodeName = TEXT("Run Dynamic State Tree");
+	bCreateNodeInstance = true;
+	bTickIntervals = true;
+}
+
+EBTNodeResult::Type UBTTask_RunDynamicStateTree::ExecuteTask(UBehaviorTreeComponent& OwnerComp, uint8* NodeMemory)
+{
+	if (!StateTreeRef.IsValid())
+	{
+		return EBTNodeResult::Failed;
+	}
+
+	FStateTreeExecutionContext Context(*OwnerComp.GetOwner(), *StateTreeRef.GetStateTree(), InstanceData);
+	SetContextDataDelegate.ExecuteIfBound(Context, OwnerComp, InjectionTag);
+	const EStateTreeRunStatus StartStatus = Context.Start(&StateTreeRef.GetParameters());
+	return GameplayStateTreeBTUtils::StateTreeRunStatusToBTNodeResult(StartStatus);
+}
+
+void UBTTask_RunDynamicStateTree::TickTask(UBehaviorTreeComponent& OwnerComp, uint8* NodeMemory, float DeltaSeconds)
+{
+	if (!StateTreeRef.IsValid())
+	{
+		FinishLatentTask(OwnerComp, EBTNodeResult::Failed);
+	}
+
+	FStateTreeExecutionContext Context(*OwnerComp.GetOwner(), *StateTreeRef.GetStateTree(), InstanceData);
+	SetContextDataDelegate.ExecuteIfBound(Context, OwnerComp, InjectionTag);
+	const EStateTreeRunStatus TickStatus = Context.Tick(DeltaSeconds);
+	if (TickStatus != EStateTreeRunStatus::Running)
+	{
+		FinishLatentTask(OwnerComp, GameplayStateTreeBTUtils::StateTreeRunStatusToBTNodeResult(TickStatus));
+	}
+	else
+	{
+		SetNextTickTime(NodeMemory, FMath::Max(0.f, Interval + FMath::FRandRange(-RandomDeviation, RandomDeviation)));
+	}
+}
+
+void UBTTask_RunDynamicStateTree::OnTaskFinished(UBehaviorTreeComponent& OwnerComp, uint8* NodeMemory, EBTNodeResult::Type TaskResult)
+{
+	if (StateTreeRef.IsValid())
+	{
+		FStateTreeExecutionContext Context(*OwnerComp.GetOwner(), *StateTreeRef.GetStateTree(), InstanceData);
+		SetContextDataDelegate.ExecuteIfBound(Context, OwnerComp, InjectionTag);
+		Context.Stop();
+	}
+}
+
+void UBTTask_RunDynamicStateTree::SetStateTreeToRun(UBehaviorTreeComponent& OwnerComp, const FStateTreeReference& InStateTree, const FSetContextDataDelegate& InSetContextDelegate, float InInterval, float InRandomDeviation)
+{
+	if (StateTreeRef.IsValid())
+	{
+		FStateTreeExecutionContext Context(*OwnerComp.GetOwner(), *StateTreeRef.GetStateTree(), InstanceData);
+		if (Context.GetStateTreeRunStatus() == EStateTreeRunStatus::Running)
+		{
+			SetContextDataDelegate.ExecuteIfBound(Context, OwnerComp, InjectionTag);
+			Context.Stop();
+		}
+	}
+
+	StateTreeRef = InStateTree;
+	SetContextDataDelegate  = InSetContextDelegate;
+	Interval = InInterval;
+	RandomDeviation = InRandomDeviation;
+}
+
+FString UBTTask_RunDynamicStateTree::GetStaticDescription() const
+{
+	return FString::Printf(TEXT("%s: %s"), *Super::GetStaticDescription(), *InjectionTag.ToString());
+}
+
+void UBTTask_RunDynamicStateTree::DescribeRuntimeValues(const UBehaviorTreeComponent& OwnerComp, uint8* NodeMemory, EBTDescriptionVerbosity::Type Verbosity, TArray<FString>& Values) const
+{
+	Super::DescribeRuntimeValues(OwnerComp, NodeMemory, Verbosity, Values);
+	Values.Add(FString::Printf(TEXT("state tree: %s"), *GetNameSafe(StateTreeRef.GetStateTree())));
+	if (StateTreeRef.IsValid())
+	{
+		Values.Add(FString::Printf(TEXT("Interval: %.2f RandomDeviation: %.2f"), Interval, RandomDeviation));
+	}
+}
\ No newline at end of file
diff --git a/Engine/Plugins/Runtime/GameplayStateTree/Source/GameplayStateTreeModule/Private/BehaviorTree/Tasks/BTTask_RunStateTree.cpp b/Engine/Plugins/Runtime/GameplayStateTree/Source/GameplayStateTreeModule/Private/BehaviorTree/Tasks/BTTask_RunStateTree.cpp
new file mode 100644
--- /dev/null
+++ b/Engine/Plugins/Runtime/GameplayStateTree/Source/GameplayStateTreeModule/Private/BehaviorTree/Tasks/BTTask_RunStateTree.cpp
@@ -0,0 +1,66 @@
+// Copyright Epic Games, Inc. All Rights Reserved.
+
+#include "BehaviorTree/Tasks/BTTask_RunStateTree.h"
+
+#include "BehaviorTree/GameplayStateTreeBTUtils.h"
+#include "Components/StateTreeComponentSchema.h"
+#include "StateTreeExecutionContext.h"
+
+UBTTask_RunStateTree::UBTTask_RunStateTree(const FObjectInitializer& ObjectInitializer /*= FObjectInitializer::Get()*/)
+	: Super(ObjectInitializer)
+{
+	INIT_TASK_NODE_NOTIFY_FLAGS();
+	NodeName = TEXT("Run State Tree");
+	bCreateNodeInstance = true;
+	bTickIntervals = true;
+}
+
+EBTNodeResult::Type UBTTask_RunStateTree::ExecuteTask(UBehaviorTreeComponent& OwnerComp, uint8* NodeMemory)
+{
+	if (StateTreeRef.IsValid())
+	{
+		FStateTreeExecutionContext Context(*OwnerComp.GetOwner(), *StateTreeRef.GetStateTree(), InstanceData);
+		if (UStateTreeComponentSchema::SetContextRequirements(OwnerComp, Context))
+		{
+			const EStateTreeRunStatus StartStatus = Context.Start(&StateTreeRef.GetParameters());
+			return GameplayStateTreeBTUtils::StateTreeRunStatusToBTNodeResult(StartStatus);
+		}
+	}
+
+	return EBTNodeResult::Failed;
+}
+
+void UBTTask_RunStateTree::TickTask(UBehaviorTreeComponent& OwnerComp, uint8* NodeMemory, float DeltaSeconds)
+{
+	if (StateTreeRef.IsValid())
+	{
+		FStateTreeExecutionContext Context(*OwnerComp.GetOwner(), *StateTreeRef.GetStateTree(), InstanceData);
+		if (UStateTreeComponentSchema::SetContextRequirements(OwnerComp, Context))
+		{
+			const EStateTreeRunStatus TickStatus = Context.Tick(DeltaSeconds);
+			if (TickStatus != EStateTreeRunStatus::Running)
+			{
+				FinishLatentTask(OwnerComp, GameplayStateTreeBTUtils::StateTreeRunStatusToBTNodeResult(TickStatus));
+			}
+			else
+			{
+				SetNextTickTime(NodeMemory, FMath::Max(0.f, Interval + FMath::FRandRange(-RandomDeviation, RandomDeviation)));
+			}
+			return;
+		}
+	}
+
+	FinishLatentTask(OwnerComp, EBTNodeResult::Failed);
+}
+
+void UBTTask_RunStateTree::OnTaskFinished(UBehaviorTreeComponent& OwnerComp, uint8* NodeMemory, EBTNodeResult::Type TaskResult)
+{
+	if (StateTreeRef.IsValid())
+	{
+		FStateTreeExecutionContext Context(*OwnerComp.GetOwner(), *StateTreeRef.GetStateTree(), InstanceData);
+		if(UStateTreeComponentSchema::SetContextRequirements(OwnerComp, Context))
+		{
+			Context.Stop();
+		}
+	}
+}
diff --git a/Engine/Plugins/Runtime/GameplayStateTree/Source/GameplayStateTreeModule/Private/Components/StateTreeComponent.cpp b/Engine/Plugins/Runtime/GameplayStateTree/Source/GameplayStateTreeModule/Private/Components/StateTreeComponent.cpp
--- a/Engine/Plugins/Runtime/GameplayStateTree/Source/GameplayStateTreeModule/Private/Components/StateTreeComponent.cpp
+++ b/Engine/Plugins/Runtime/GameplayStateTree/Source/GameplayStateTreeModule/Private/Components/StateTreeComponent.cpp
@@ -60,101 +60,12 @@ void UStateTreeComponent::UninitializeComponent()
 
 bool UStateTreeComponent::CollectExternalData(const FStateTreeExecutionContext& Context, const UStateTree* StateTree, TArrayView<const FStateTreeExternalDataDesc> ExternalDataDescs, TArrayView<FStateTreeDataView> OutDataViews) const
 {
-	const UWorld* World = GetWorld();
-	if (World == nullptr)
-	{
-		return false;
-	}
-
-	check(ExternalDataDescs.Num() == OutDataViews.Num());
-	
-	for (int32 Index = 0; Index < ExternalDataDescs.Num(); Index++)
-	{
-		const FStateTreeExternalDataDesc& ItemDesc = ExternalDataDescs[Index];
-		if (ItemDesc.Struct != nullptr)
-		{
-			if (ItemDesc.Struct->IsChildOf(UWorldSubsystem::StaticClass()))
-			{
-				UWorldSubsystem* Subsystem = World->GetSubsystemBase(Cast<UClass>(const_cast<UStruct*>(ItemDesc.Struct.Get())));
-				OutDataViews[Index] = FStateTreeDataView(Subsystem);
-			}
-			else if (ItemDesc.Struct->IsChildOf(UActorComponent::StaticClass()))
-			{
-				UActorComponent* Component = GetOwner()->FindComponentByClass(Cast<UClass>(const_cast<UStruct*>(ItemDesc.Struct.Get())));
-				OutDataViews[Index] = FStateTreeDataView(Component);
-			}
-			else if (ItemDesc.Struct->IsChildOf(APawn::StaticClass()))
-			{
-				APawn* OwnerPawn = (AIOwner != nullptr) ? AIOwner->GetPawn() : Cast<APawn>(GetOwner());
-				OutDataViews[Index] = FStateTreeDataView(OwnerPawn);
-			}
-			else if (ItemDesc.Struct->IsChildOf(AAIController::StaticClass()))
-			{
-				AAIController* OwnerController = (AIOwner != nullptr) ? AIOwner.Get() : Cast<AAIController>(GetOwner());
-				OutDataViews[Index] = FStateTreeDataView(OwnerController);
-			}
-			else if (ItemDesc.Struct->IsChildOf(AActor::StaticClass()))
-			{
-				AActor* OwnerActor = (AIOwner != nullptr) ? AIOwner->GetPawn() : GetOwner();
-				OutDataViews[Index] = FStateTreeDataView(OwnerActor);
-			}
-		}
-	}
-
-	return true;
+	return UStateTreeComponentSchema::CollectExternalData(Context, StateTree, ExternalDataDescs, OutDataViews);
 }
 
 bool UStateTreeComponent::SetContextRequirements(FStateTreeExecutionContext& Context, bool bLogErrors)
 {
-	if (!Context.IsValid())
-	{
-		return false;
-	}
-
-	Context.SetCollectExternalDataCallback(FOnCollectStateTreeExternalData::CreateUObject(this, &UStateTreeComponent::CollectExternalData));
-	
-	// Make sure the actor matches one required.
-	AActor* ContextActor = nullptr;
-	const UStateTreeComponentSchema* Schema = Cast<UStateTreeComponentSchema>(Context.GetStateTree()->GetSchema());
-	if (Schema)
-	{
-		if (AAIController* OwnerController = (AIOwner != nullptr) ? AIOwner.Get() : Cast<AAIController>(GetOwner()))
-		{
-			if (OwnerController && OwnerController->IsA(Schema->GetContextActorClass()))
-			{
-				ContextActor = OwnerController;
-			}
-		}
-		if (ContextActor == nullptr)
-		{
-			if (AActor* OwnerActor = (AIOwner != nullptr) ? AIOwner->GetPawn() : GetOwner())
-			{
-				if (OwnerActor && OwnerActor->IsA(Schema->GetContextActorClass()))
-				{
-					ContextActor = OwnerActor;
-				}
-			}
-		}
-		if (ContextActor == nullptr && bLogErrors)
-		{
-			STATETREE_LOG(Error, TEXT("%s: Could not find context actor of type %s. StateTree will not update."), ANSI_TO_TCHAR(__FUNCTION__), *GetNameSafe(Schema->GetContextActorClass()));
-		}
-	}
-	else if (bLogErrors)
-	{
-		STATETREE_LOG(Error, TEXT("%s: Expected StateTree asset to contain StateTreeComponentSchema. StateTree will not update."), ANSI_TO_TCHAR(__FUNCTION__));
-	}
-	
-	const FName ActorName(TEXT("Actor"));
-	Context.SetContextDataByName(ActorName, FStateTreeDataView(ContextActor));
-
-	bool bResult = Context.AreContextDataViewsValid();
-	if (!bResult && bLogErrors)
-	{
-		STATETREE_LOG(Error, TEXT("%s: Missing external data requirements. StateTree will not update."), ANSI_TO_TCHAR(__FUNCTION__));
-	}
-	
-	return bResult;
+	return UStateTreeComponentSchema::SetContextRequirements(*this, Context);
 }
 
 void UStateTreeComponent::BeginPlay()
@@ -190,7 +101,7 @@ void UStateTreeComponent::TickComponent(float DeltaTime, enum ELevelTick TickTyp
 	}
 
 	FStateTreeExecutionContext Context(*GetOwner(), *StateTreeRef.GetStateTree(), InstanceData);
-	if (SetContextRequirements(Context))
+	if (UStateTreeComponentSchema::SetContextRequirements(*this, Context))
 	{
 		const EStateTreeRunStatus PreviousRunStatus = Context.GetStateTreeRunStatus();
 		const EStateTreeRunStatus CurrentRunStatus = Context.Tick(DeltaTime);
@@ -213,7 +124,7 @@ void UStateTreeComponent::StartLogic()
 	}
 
 	FStateTreeExecutionContext Context(*GetOwner(), *StateTreeRef.GetStateTree(), InstanceData);
-	if (SetContextRequirements(Context))
+	if (UStateTreeComponentSchema::SetContextRequirements(*this, Context))
 	{
 		const EStateTreeRunStatus PreviousRunStatus = Context.GetStateTreeRunStatus();
 		const EStateTreeRunStatus CurrentRunStatus = Context.Start(&StateTreeRef.GetParameters());
@@ -237,7 +148,7 @@ void UStateTreeComponent::RestartLogic()
 	}
 
 	FStateTreeExecutionContext Context(*GetOwner(), *StateTreeRef.GetStateTree(), InstanceData);
-	if (SetContextRequirements(Context))
+	if (UStateTreeComponentSchema::SetContextRequirements(*this, Context))
 	{
 		const EStateTreeRunStatus PreviousRunStatus = Context.GetStateTreeRunStatus();
 		const EStateTreeRunStatus CurrentRunStatus = Context.Start(&StateTreeRef.GetParameters());
@@ -266,7 +177,7 @@ void UStateTreeComponent::StopLogic(const FString& Reason)
 	}
 
 	FStateTreeExecutionContext Context(*GetOwner(), *StateTreeRef.GetStateTree(), InstanceData);
-	if (SetContextRequirements(Context))
+	if (UStateTreeComponentSchema::SetContextRequirements(*this, Context))
 	{
 		const EStateTreeRunStatus PreviousRunStatus = Context.GetStateTreeRunStatus();
 		const EStateTreeRunStatus CurrentRunStatus = Context.Stop();
diff --git a/Engine/Plugins/Runtime/GameplayStateTree/Source/GameplayStateTreeModule/Private/Components/StateTreeComponentSchema.cpp b/Engine/Plugins/Runtime/GameplayStateTree/Source/GameplayStateTreeModule/Private/Components/StateTreeComponentSchema.cpp
--- a/Engine/Plugins/Runtime/GameplayStateTree/Source/GameplayStateTreeModule/Private/Components/StateTreeComponentSchema.cpp
+++ b/Engine/Plugins/Runtime/GameplayStateTree/Source/GameplayStateTreeModule/Private/Components/StateTreeComponentSchema.cpp
@@ -1,11 +1,16 @@
 ﻿// Copyright Epic Games, Inc. All Rights Reserved.
 
 #include "Components/StateTreeComponentSchema.h"
+#include "AIController.h"
+#include "BrainComponent.h"
+#include "Engine/World.h"
+#include "GameFramework/Pawn.h"
 #include "StateTreeConditionBase.h"
 #include "StateTreeEvaluatorBase.h"
+#include "StateTreeExecutionContext.h"
 #include "StateTreeTaskBase.h"
 #include "Subsystems/WorldSubsystem.h"
-
+#include "VisualLogger/VisualLogger.h"
 
 UStateTreeComponentSchema::UStateTreeComponentSchema()
 	: ContextActorClass(AActor::StaticClass())
@@ -60,3 +65,110 @@ void UStateTreeComponentSchema::PostEditChangeChainProperty(FPropertyChangedChai
 	}
 }
 #endif
+
+bool UStateTreeComponentSchema::SetContextRequirements(UBrainComponent& BrainComponent, FStateTreeExecutionContext& Context, bool bLogErrors /*= false*/)
+{
+	if (!Context.IsValid())
+	{
+		return false;
+	}
+
+	Context.SetCollectExternalDataCallback(FOnCollectStateTreeExternalData::CreateStatic(&UStateTreeComponentSchema::CollectExternalData));
+
+	// Make sure the actor matches one required.
+	AActor* ContextActor = nullptr;
+	const UStateTreeComponentSchema* Schema = Cast<UStateTreeComponentSchema>(Context.GetStateTree()->GetSchema());
+	if (Schema)
+	{
+		AAIController* AIOwner = BrainComponent.GetAIOwner();
+		if (AAIController* OwnerController = (AIOwner != nullptr) ? AIOwner : Cast<AAIController>(BrainComponent.GetOwner()))
+		{
+			if (OwnerController && OwnerController->IsA(Schema->GetContextActorClass()))
+			{
+				ContextActor = OwnerController;
+			}
+		}
+		if (ContextActor == nullptr)
+		{
+			if (AActor* OwnerActor = (AIOwner != nullptr) ? AIOwner->GetPawn() : BrainComponent.GetOwner())
+			{
+				if (OwnerActor && OwnerActor->IsA(Schema->GetContextActorClass()))
+				{
+					ContextActor = OwnerActor;
+				}
+			}
+		}
+		if (ContextActor == nullptr && bLogErrors)
+		{
+			UE_VLOG(BrainComponent.GetOwner(), LogStateTree, Error, TEXT("%s: Could not find context actor of type %s. StateTree will not update."), ANSI_TO_TCHAR(__FUNCTION__), *GetNameSafe(Schema->GetContextActorClass()));
+		}
+	}
+	else if (bLogErrors)
+	{
+		UE_VLOG(BrainComponent.GetOwner(), LogStateTree, Error, TEXT("%s: Expected StateTree asset to contain StateTreeComponentSchema. StateTree will not update."), ANSI_TO_TCHAR(__FUNCTION__));
+	}
+
+	const FName ActorName(TEXT("Actor"));
+	Context.SetContextDataByName(ActorName, FStateTreeDataView(ContextActor));
+
+	bool bResult = Context.AreContextDataViewsValid();
+	if (!bResult && bLogErrors)
+	{
+		UE_VLOG(BrainComponent.GetOwner(), LogStateTree, Error, TEXT("%s: Missing external data requirements. StateTree will not update."), ANSI_TO_TCHAR(__FUNCTION__));
+	}
+
+	return bResult;
+}
+
+bool UStateTreeComponentSchema::CollectExternalData(const FStateTreeExecutionContext& Context, const UStateTree* StateTree, TArrayView<const FStateTreeExternalDataDesc> ExternalDataDescs, TArrayView<FStateTreeDataView> OutDataViews)
+{
+	const UWorld* World = Context.GetWorld();
+	if (World == nullptr)
+	{
+		return false;
+	}
+
+	check(ExternalDataDescs.Num() == OutDataViews.Num());
+
+	AActor* Owner = Cast<AActor>(Context.GetOwner());
+	if (!Owner)
+	{
+		return false;
+	}
+
+	AAIController* AIOwner = Cast<AAIController>(Owner);
+	for (int32 Index = 0; Index < ExternalDataDescs.Num(); Index++)
+	{
+		const FStateTreeExternalDataDesc& ItemDesc = ExternalDataDescs[Index];
+		if (ItemDesc.Struct != nullptr)
+		{
+			if (ItemDesc.Struct->IsChildOf(UWorldSubsystem::StaticClass()))
+			{
+				UWorldSubsystem* Subsystem = World->GetSubsystemBase(Cast<UClass>(const_cast<UStruct*>(ItemDesc.Struct.Get())));
+				OutDataViews[Index] = FStateTreeDataView(Subsystem);
+			}
+			else if (ItemDesc.Struct->IsChildOf(UActorComponent::StaticClass()))
+			{
+				UActorComponent* Component = Owner->FindComponentByClass(Cast<UClass>(const_cast<UStruct*>(ItemDesc.Struct.Get())));
+				OutDataViews[Index] = FStateTreeDataView(Component);
+			}
+			else if (ItemDesc.Struct->IsChildOf(APawn::StaticClass()))
+			{
+				APawn* OwnerPawn = (AIOwner != nullptr) ? AIOwner->GetPawn() : Cast<APawn>(Owner);
+				OutDataViews[Index] = FStateTreeDataView(OwnerPawn);
+			}
+			else if (ItemDesc.Struct->IsChildOf(AAIController::StaticClass()))
+			{
+				AAIController* OwnerController = AIOwner;
+				OutDataViews[Index] = FStateTreeDataView(OwnerController);
+			}
+			else if (ItemDesc.Struct->IsChildOf(AActor::StaticClass()))
+			{
+				AActor* OwnerActor = (AIOwner != nullptr) ? AIOwner->GetPawn() : Owner;
+				OutDataViews[Index] = FStateTreeDataView(OwnerActor);
+			}
+		}
+	}
+
+	return true;
+}
\ No newline at end of file
diff --git a/Engine/Plugins/Runtime/GameplayStateTree/Source/GameplayStateTreeModule/Public/BehaviorTree/GameplayStateTreeBTUtils.h b/Engine/Plugins/Runtime/GameplayStateTree/Source/GameplayStateTreeModule/Public/BehaviorTree/GameplayStateTreeBTUtils.h
new file mode 100644
--- /dev/null
+++ b/Engine/Plugins/Runtime/GameplayStateTree/Source/GameplayStateTreeModule/Public/BehaviorTree/GameplayStateTreeBTUtils.h
@@ -0,0 +1,12 @@
+// Copyright Epic Games, Inc. All Rights Reserved.
+
+#pragma once
+
+#include "BehaviorTree/BehaviorTreeTypes.h"
+
+enum class EStateTreeRunStatus : uint8;
+
+namespace GameplayStateTreeBTUtils
+{
+	GAMEPLAYSTATETREEMODULE_API EBTNodeResult::Type StateTreeRunStatusToBTNodeResult(EStateTreeRunStatus Result);
+}
\ No newline at end of file
diff --git a/Engine/Plugins/Runtime/GameplayStateTree/Source/GameplayStateTreeModule/Public/BehaviorTree/Tasks/BTTask_RunDynamicStateTree.h b/Engine/Plugins/Runtime/GameplayStateTree/Source/GameplayStateTreeModule/Public/BehaviorTree/Tasks/BTTask_RunDynamicStateTree.h
new file mode 100644
--- /dev/null
+++ b/Engine/Plugins/Runtime/GameplayStateTree/Source/GameplayStateTreeModule/Public/BehaviorTree/Tasks/BTTask_RunDynamicStateTree.h
@@ -0,0 +1,56 @@
+// Copyright Epic Games, Inc. All Rights Reserved.
+
+#pragma once
+
+#include "BehaviorTree/BTTaskNode.h"
+#include "StateTreeReference.h"
+#include "StateTreeInstanceData.h"
+
+#include "BTTask_RunDynamicStateTree.generated.h"
+
+enum class EStateTreeRunStatus : uint8;
+
+struct FStateTreeExecutionContext;
+
+/**
+ * RunDynamicStateTree task allows the execution of a state tree chosen at runtime.
+ * UBTTask_RunDynamicStateTree::SetDynamicStateTree can be used to set the node's state tree.
+ */
+UCLASS(MinimalAPI)
+class UBTTask_RunDynamicStateTree : public UBTTaskNode
+{
+	GENERATED_BODY()
+public:
+	DECLARE_DELEGATE_ThreeParams(FSetContextDataDelegate, FStateTreeExecutionContext&, UBehaviorTreeComponent&, FGameplayTag InjectionTag);
+
+	GAMEPLAYSTATETREEMODULE_API static bool SetDynamicStateTree(UBehaviorTreeComponent& OwnerComp, FGameplayTag InInjectTag, const FStateTreeReference& InStateTree, const FSetContextDataDelegate& InSetContextDataDelegate, float InInterval, float InRandomDeviation, UBTCompositeNode* OptionalStartNode = nullptr);
+
+	GAMEPLAYSTATETREEMODULE_API FGameplayTag GetInjectionTag() const { return InjectionTag; }
+
+protected:
+	GAMEPLAYSTATETREEMODULE_API UBTTask_RunDynamicStateTree(const FObjectInitializer& ObjectInitializer = FObjectInitializer::Get());
+
+	GAMEPLAYSTATETREEMODULE_API virtual FString GetStaticDescription() const override;
+	GAMEPLAYSTATETREEMODULE_API virtual void DescribeRuntimeValues(const UBehaviorTreeComponent& OwnerComp, uint8* NodeMemory, EBTDescriptionVerbosity::Type Verbosity, TArray<FString>& Values) const override;
+
+	GAMEPLAYSTATETREEMODULE_API virtual EBTNodeResult::Type ExecuteTask(UBehaviorTreeComponent& OwnerComp, uint8* NodeMemory) override;
+	GAMEPLAYSTATETREEMODULE_API virtual void TickTask(UBehaviorTreeComponent& OwnerComp, uint8* NodeMemory, float DeltaSeconds) override;
+	GAMEPLAYSTATETREEMODULE_API virtual void OnTaskFinished(UBehaviorTreeComponent& OwnerComp, uint8* NodeMemory, EBTNodeResult::Type TaskResult) override;
+
+	void SetStateTreeToRun(UBehaviorTreeComponent& OwnerComp, const FStateTreeReference& StateTreeToRun, const FSetContextDataDelegate& SetContextDelegate, float Interval, float RandomDeviation);
+
+	UPROPERTY(Transient)
+	FStateTreeReference StateTreeRef;
+
+	UPROPERTY(Transient)
+	FStateTreeInstanceData InstanceData;
+
+	/** Gameplay tag that will identify this task for state tree injection */
+	UPROPERTY(EditAnywhere, Category = Node)
+	FGameplayTag InjectionTag;
+
+	FSetContextDataDelegate SetContextDataDelegate;
+
+	float Interval = 1.f;
+	float RandomDeviation = 0.f;
+};
\ No newline at end of file
diff --git a/Engine/Plugins/Runtime/GameplayStateTree/Source/GameplayStateTreeModule/Public/BehaviorTree/Tasks/BTTask_RunStateTree.h b/Engine/Plugins/Runtime/GameplayStateTree/Source/GameplayStateTreeModule/Public/BehaviorTree/Tasks/BTTask_RunStateTree.h
new file mode 100644
--- /dev/null
+++ b/Engine/Plugins/Runtime/GameplayStateTree/Source/GameplayStateTreeModule/Public/BehaviorTree/Tasks/BTTask_RunStateTree.h
@@ -0,0 +1,40 @@
+// Copyright Epic Games, Inc. All Rights Reserved.
+
+#pragma once
+
+#include "BehaviorTree/BTTaskNode.h"
+#include "StateTreeReference.h"
+#include "StateTreeInstanceData.h"
+
+#include "BTTask_RunStateTree.generated.h"
+
+/**
+ * RunStateTree task allows the execution of state tree with the StateTreeComponentSchema inside a behavior tree.
+ */
+ UCLASS(MinimalAPI)
+class UBTTask_RunStateTree : public UBTTaskNode
+{
+	GENERATED_BODY()
+public:
+	GAMEPLAYSTATETREEMODULE_API UBTTask_RunStateTree(const FObjectInitializer& ObjectInitializer = FObjectInitializer::Get());
+
+protected:
+	GAMEPLAYSTATETREEMODULE_API virtual EBTNodeResult::Type ExecuteTask(UBehaviorTreeComponent& OwnerComp, uint8* NodeMemory) override;
+	GAMEPLAYSTATETREEMODULE_API virtual void TickTask(UBehaviorTreeComponent& OwnerComp, uint8* NodeMemory, float DeltaSeconds) override;
+	GAMEPLAYSTATETREEMODULE_API virtual void OnTaskFinished(UBehaviorTreeComponent& OwnerComp, uint8* NodeMemory, EBTNodeResult::Type TaskResult) override;
+	
+	/** State tree that will be run when the task is selected. */
+	UPROPERTY(EditAnywhere, Category = Task, meta = (Schema = "/Script/GameplayStateTreeModule.StateTreeComponentSchema"))
+	FStateTreeReference StateTreeRef;
+
+	UPROPERTY(Transient)
+	FStateTreeInstanceData InstanceData;
+
+	/** Interval between state tree update. */
+	UPROPERTY(EditAnywhere, Category = Task, meta = (ClampMin = "0.001"))
+	float Interval = 0.01f;
+
+	/** Random deviation on the interval between each state tree update. */
+	UPROPERTY(EditAnywhere, Category = Task, meta = (ClampMin = "0.0"))
+	float RandomDeviation = 0.0f;
+};
diff --git a/Engine/Plugins/Runtime/GameplayStateTree/Source/GameplayStateTreeModule/Public/Components/StateTreeComponent.h b/Engine/Plugins/Runtime/GameplayStateTree/Source/GameplayStateTreeModule/Public/Components/StateTreeComponent.h
--- a/Engine/Plugins/Runtime/GameplayStateTree/Source/GameplayStateTreeModule/Public/Components/StateTreeComponent.h
+++ b/Engine/Plugins/Runtime/GameplayStateTree/Source/GameplayStateTreeModule/Public/Components/StateTreeComponent.h
@@ -59,7 +59,7 @@ public:
 	 */
 	UFUNCTION(BlueprintCallable, Category = "Gameplay|StateTree")
 	void SetStartLogicAutomatically(const bool bInStartLogicAutomatically);
-
+	 
 	/** Sends event to the running StateTree. */
 	UFUNCTION(BlueprintCallable, Category = "Gameplay|StateTree")
 	void SendStateTreeEvent(const FStateTreeEvent& Event);
@@ -85,8 +85,10 @@ protected:
 	virtual void PostLoad() override;
 #endif
 	
+	UE_DEPRECATED(5.4, "Please use UStateTreeComponentSchema::SetContextRequirements instead.")
 	bool SetContextRequirements(FStateTreeExecutionContext& Context, bool bLogErrors = false);
 	
+	UE_DEPRECATED(5.4, "Please use UStateTreeComponentSchema::CollectExternalData instead.")
 	bool CollectExternalData(const FStateTreeExecutionContext& Context, const UStateTree* StateTree, TArrayView<const FStateTreeExternalDataDesc> Descs, TArrayView<FStateTreeDataView> OutDataViews) const;
 	
 #if WITH_EDITORONLY_DATA
diff --git a/Engine/Plugins/Runtime/GameplayStateTree/Source/GameplayStateTreeModule/Public/Components/StateTreeComponentSchema.h b/Engine/Plugins/Runtime/GameplayStateTree/Source/GameplayStateTreeModule/Public/Components/StateTreeComponentSchema.h
--- a/Engine/Plugins/Runtime/GameplayStateTree/Source/GameplayStateTreeModule/Public/Components/StateTreeComponentSchema.h
+++ b/Engine/Plugins/Runtime/GameplayStateTree/Source/GameplayStateTreeModule/Public/Components/StateTreeComponentSchema.h
@@ -7,6 +7,11 @@
 #include "StateTreeExecutionTypes.h"
 #include "StateTreeComponentSchema.generated.h"
 
+class UBrainComponent;
+class UStateTree;
+
+struct FStateTreeExecutionContext;
+
 /**
  * StateTree for Actors with StateTree component. 
  */
@@ -20,6 +25,9 @@ public:
 
 	UClass* GetContextActorClass() const { return ContextActorClass; };
 	
+	static bool SetContextRequirements(UBrainComponent& BrainComponent, FStateTreeExecutionContext& Context, bool bLogErrors = false);
+	static bool CollectExternalData(const FStateTreeExecutionContext& Context, const UStateTree* StateTree, TArrayView<const FStateTreeExternalDataDesc> Descs, TArrayView<FStateTreeDataView> OutDataViews);
+
 protected:
 	virtual bool IsStructAllowed(const UScriptStruct* InScriptStruct) const override;
 	virtual bool IsClassAllowed(const UClass* InScriptStruct) const override;
