Create StateTreeMoveToTask wrapping a UAITask_MoveTo allowing state tree to move AI using the same framework as BehaviorTree
Create the FStateTreeAITaskBase & FStateTreeAIActionTaskBase state tree task namespace for AI focused ST task

diff --git a/Engine/Plugins/Runtime/GameplayStateTree/Source/GameplayStateTreeModule/GameplayStateTreeModule.Build.cs b/Engine/Plugins/Runtime/GameplayStateTree/Source/GameplayStateTreeModule/GameplayStateTreeModule.Build.cs
--- a/Engine/Plugins/Runtime/GameplayStateTree/Source/GameplayStateTreeModule/GameplayStateTreeModule.Build.cs
+++ b/Engine/Plugins/Runtime/GameplayStateTree/Source/GameplayStateTreeModule/GameplayStateTreeModule.Build.cs
@@ -19,6 +19,7 @@ namespace UnrealBuildTool.Rules
 					"Engine",
 					"GameplayTags",
 					"GameplayTasks",
+					"NavigationSystem",
 					"StateTreeModule",
 					"StructUtils"
 				}
diff --git a/Engine/Plugins/Runtime/GameplayStateTree/Source/GameplayStateTreeModule/Private/Tasks/StateTreeMoveToTask.cpp b/Engine/Plugins/Runtime/GameplayStateTree/Source/GameplayStateTreeModule/Private/Tasks/StateTreeMoveToTask.cpp
new file mode 100644
--- /dev/null
+++ b/Engine/Plugins/Runtime/GameplayStateTree/Source/GameplayStateTreeModule/Private/Tasks/StateTreeMoveToTask.cpp
@@ -0,0 +1,132 @@
+// Copyright Epic Games, Inc. All Rights Reserved.
+
+#include "Tasks/StateTreeMoveToTask.h"
+
+#include "AIController.h"
+#include "Navigation/PathFollowingComponent.h"
+#include "StateTreeExecutionContext.h"
+#include "StateTreeLinker.h"
+#include "Tasks/AITask_MoveTo.h"
+#include "VisualLogger/VisualLogger.h"
+
+#include UE_INLINE_GENERATED_CPP_BY_NAME(StateTreeMoveToTask)
+
+EStateTreeRunStatus FStateTreeMoveToTask::EnterState(FStateTreeExecutionContext& Context, const FStateTreeTransitionResult& Transition) const
+{
+	FInstanceDataType& InstanceData = Context.GetInstanceData(*this);
+	if (!InstanceData.AIController)
+	{
+		UE_VLOG(Context.GetOwner(), LogStateTree, Error, TEXT("FStateTreeMoveToTask failed since AIController is missing."));
+		return EStateTreeRunStatus::Failed;
+	}
+
+	InstanceData.TaskOwner = TScriptInterface<IGameplayTaskOwnerInterface>(InstanceData.AIController->FindComponentByInterface(UGameplayTaskOwnerInterface::StaticClass()));
+	if (!InstanceData.TaskOwner)
+	{
+		InstanceData.TaskOwner = InstanceData.AIController;
+	}
+
+	return PerformMoveTask(Context, *InstanceData.AIController);
+}
+
+EStateTreeRunStatus FStateTreeMoveToTask::Tick(FStateTreeExecutionContext& Context, const float DeltaTime) const
+{
+	FInstanceDataType& InstanceData = Context.GetInstanceData(*this);
+	if (InstanceData.MoveToTask)
+	{
+		if (InstanceData.MoveToTask->GetState() == EGameplayTaskState::Finished)
+		{
+			return InstanceData.MoveToTask->WasMoveSuccessful() ? EStateTreeRunStatus::Succeeded : EStateTreeRunStatus::Failed;
+		}
+
+		if (InstanceData.bTrackMovingGoal && !InstanceData.TargetActor)
+		{
+			const FVector CurrentDestination = InstanceData.MoveToTask->GetMoveRequestRef().GetDestination();
+			if (FVector::DistSquared(CurrentDestination, InstanceData.Destination) > (InstanceData.DestinationMoveTolerance * InstanceData.DestinationMoveTolerance))
+			{
+				UE_VLOG(Context.GetOwner(), LogStateTree, Log, TEXT("FStateTreeMoveToTask destination has moved enough. Restarting task."));
+				return PerformMoveTask(Context, *InstanceData.AIController);
+			}
+		}
+		return EStateTreeRunStatus::Running;
+	}
+	return EStateTreeRunStatus::Failed;
+}
+
+void FStateTreeMoveToTask::ExitState(FStateTreeExecutionContext& Context, const FStateTreeTransitionResult& Transition) const
+{
+	FInstanceDataType& InstanceData = Context.GetInstanceData(*this);
+	if (InstanceData.MoveToTask && InstanceData.MoveToTask->GetState() != EGameplayTaskState::Finished)
+	{
+		UE_VLOG(Context.GetOwner(), LogStateTree, Log, TEXT("FStateTreeMoveToTask aborting move to because state finished."));
+		InstanceData.MoveToTask->ExternalCancel();
+	}
+}
+
+UAITask_MoveTo* FStateTreeMoveToTask::PrepareMoveToTask(FStateTreeExecutionContext& Context, AAIController& Controller, UAITask_MoveTo* ExistingTask, FAIMoveRequest& MoveRequest) const
+{
+	FInstanceDataType& InstanceData = Context.GetInstanceData(*this);
+	UAITask_MoveTo* MoveTask = ExistingTask ? ExistingTask : UAITask::NewAITask<UAITask_MoveTo>(Controller, *InstanceData.TaskOwner);
+	if (MoveTask)
+	{
+		MoveTask->SetUp(&Controller, MoveRequest);
+	}
+
+	return MoveTask;
+}
+
+EStateTreeRunStatus FStateTreeMoveToTask::PerformMoveTask(FStateTreeExecutionContext& Context, AAIController& Controller) const
+{
+	FInstanceDataType& InstanceData = Context.GetInstanceData(*this);
+	FAIMoveRequest MoveReq;
+	MoveReq.SetNavigationFilter(InstanceData.FilterClass ? InstanceData.FilterClass : Controller.GetDefaultNavigationFilterClass())
+		.SetAllowPartialPath(InstanceData.bAllowPartialPath)
+		.SetAcceptanceRadius(InstanceData.AcceptableRadius)
+		.SetCanStrafe(InstanceData.bAllowStrafe)
+		.SetReachTestIncludesAgentRadius(InstanceData.bReachTestIncludesAgentRadius)
+		.SetReachTestIncludesGoalRadius(InstanceData.bReachTestIncludesGoalRadius)
+		.SetRequireNavigableEndLocation(InstanceData.bRequireNavigableEndLocation)
+		.SetProjectGoalLocation(InstanceData.bProjectGoalLocation)
+		.SetUsePathfinding(true);
+
+	if (InstanceData.TargetActor)
+	{
+		if (InstanceData.bTrackMovingGoal)
+		{
+			MoveReq.SetGoalActor(InstanceData.TargetActor);
+		}
+		else
+		{
+			MoveReq.SetGoalLocation(InstanceData.TargetActor->GetActorLocation());
+		}
+	}
+	else
+	{
+		MoveReq.SetGoalLocation(InstanceData.Destination);
+	}
+
+	if (MoveReq.IsValid())
+	{	
+		if (InstanceData.MoveToTask = PrepareMoveToTask(Context, Controller, InstanceData.MoveToTask, MoveReq))
+		{
+			if (InstanceData.MoveToTask->IsActive())
+			{
+				InstanceData.MoveToTask->ConditionalPerformMove();
+			}
+			else
+			{
+				InstanceData.MoveToTask->ReadyForActivation();
+			}
+
+			if (InstanceData.MoveToTask->GetState() == EGameplayTaskState::Finished)
+			{
+				return InstanceData.MoveToTask->WasMoveSuccessful() ? EStateTreeRunStatus::Succeeded : EStateTreeRunStatus::Failed;
+			}
+
+			return EStateTreeRunStatus::Running;
+		}
+	}
+
+	UE_VLOG(Context.GetOwner(), LogStateTree, Error, TEXT("FStateTreeMoveToTask failed because it doesn't have a destination."));
+	return EStateTreeRunStatus::Failed;
+}
\ No newline at end of file
diff --git a/Engine/Plugins/Runtime/GameplayStateTree/Source/GameplayStateTreeModule/Public/Tasks/StateTreeAITask.h b/Engine/Plugins/Runtime/GameplayStateTree/Source/GameplayStateTreeModule/Public/Tasks/StateTreeAITask.h
new file mode 100644
--- /dev/null
+++ b/Engine/Plugins/Runtime/GameplayStateTree/Source/GameplayStateTreeModule/Public/Tasks/StateTreeAITask.h
@@ -0,0 +1,21 @@
+// Copyright Epic Games, Inc. All Rights Reserved.
+
+#pragma once
+
+#include "StateTreeTaskBase.h"
+
+#include "StateTreeAITask.generated.h"
+
+// Base class of all AI task that expect to be run on an AIController or derived class
+USTRUCT(meta= (Hidden, Category = "AI"))
+struct FStateTreeAITaskBase : public FStateTreeTaskBase
+{
+	GENERATED_BODY()
+};
+
+// Base class of all AI task that do a physical action 
+USTRUCT(meta = (Hidden, Category="AI|Action"))
+struct FStateTreeAIActionTaskBase : public FStateTreeAITaskBase
+{
+	GENERATED_BODY()
+};
\ No newline at end of file
diff --git a/Engine/Plugins/Runtime/GameplayStateTree/Source/GameplayStateTreeModule/Public/Tasks/StateTreeMoveToTask.h b/Engine/Plugins/Runtime/GameplayStateTree/Source/GameplayStateTreeModule/Public/Tasks/StateTreeMoveToTask.h
new file mode 100644
--- /dev/null
+++ b/Engine/Plugins/Runtime/GameplayStateTree/Source/GameplayStateTreeModule/Public/Tasks/StateTreeMoveToTask.h
@@ -0,0 +1,100 @@
+// Copyright Epic Games, Inc. All Rights Reserved.
+
+#pragma once
+
+#include "AISystem.h"
+#include "NavFilters/NavigationQueryFilter.h"
+#include "Tasks/StateTreeAITask.h"
+#include "Templates/SubclassOf.h"
+
+#include "StateTreeMovetoTask.generated.h"
+
+class AActor;
+class AAIController;
+class IGameplayTaskOwnerInterface;
+class UAITask_MoveTo;
+
+struct FAIMoveRequest;
+
+USTRUCT()
+struct FStateTreeMoveToTaskInstanceData
+{
+	GENERATED_BODY()
+
+	UPROPERTY(EditAnywhere, Category = Context)
+	TObjectPtr<AAIController> AIController = nullptr;
+
+	UPROPERTY(EditAnywhere, Category = Parameter)
+	FVector Destination = FVector::Zero();
+
+	UPROPERTY(EditAnywhere, Category = Parameter)
+	TObjectPtr<AActor> TargetActor;
+
+	/** fixed distance added to threshold between AI and goal location in destination reach test */
+	UPROPERTY(EditAnywhere, Category = Parameter, meta=(ClampMin = "0.0", UIMin="0.0"))
+	float AcceptableRadius = GET_AI_CONFIG_VAR(AcceptanceRadius);
+
+	/** if the task is expected to react to changes to location in input
+	 *	this property can be used to tweak sensitivity of the mechanism. Value is 
+	 *	recommended to be less than AcceptableRadius */
+	UPROPERTY(EditAnywhere, Category = Parameter, meta = (ClampMin = "0.0", UIMin = "0.0", EditCondition="bTrackMovingGoal"))
+	float DestinationMoveTolerance = 0.f;
+
+	/** "None" will result in default filter being used */
+	UPROPERTY(EditAnywhere, Category = Parameter)
+	TSubclassOf<UNavigationQueryFilter> FilterClass;
+
+	UPROPERTY(EditAnywhere, Category = Parameter)
+	bool bAllowStrafe = GET_AI_CONFIG_VAR(bAllowStrafing);
+
+	/** if set, use incomplete path when goal can't be reached */
+	UPROPERTY(EditAnywhere, Category = Parameter)
+	bool bAllowPartialPath = GET_AI_CONFIG_VAR(bAcceptPartialPaths);
+
+	/** if set, path to goal actor will update itself when actor moves */
+	UPROPERTY(EditAnywhere, Category = Parameter)
+	bool bTrackMovingGoal = true;
+
+	/** if set, the goal location will need to be navigable */
+	UPROPERTY(EditAnywhere, Category = Parameter)
+	bool bRequireNavigableEndLocation = true;
+
+	/** if set, goal location will be projected on navigation data (navmesh) before using */
+	UPROPERTY(EditAnywhere, Category = Parameter)
+	bool bProjectGoalLocation = true;
+
+	/** if set, radius of AI's capsule will be added to threshold between AI and goal location in destination reach test  */
+	UPROPERTY(EditAnywhere, Category = Parameter)
+	bool bReachTestIncludesAgentRadius = GET_AI_CONFIG_VAR(bFinishMoveOnGoalOverlap);
+	
+	/** if set, radius of goal's capsule will be added to threshold between AI and goal location in destination reach test  */
+	UPROPERTY(EditAnywhere, Category = Parameter)
+	bool bReachTestIncludesGoalRadius = GET_AI_CONFIG_VAR(bFinishMoveOnGoalOverlap);
+
+	UPROPERTY(Transient)
+	TObjectPtr<UAITask_MoveTo> MoveToTask = nullptr;
+
+	UPROPERTY(Transient)
+	TScriptInterface<IGameplayTaskOwnerInterface> TaskOwner = nullptr;
+};
+
+
+/**
+* Task moving the given AIController's pawn using the AITask_MoveTo framework. Succeed when move ends at destinations, fails if move is impossible.
+*/
+USTRUCT(meta = (DisplayName = "Move To", Category = "AI|Action"))
+struct FStateTreeMoveToTask : public FStateTreeAIActionTaskBase
+{
+	GENERATED_BODY()
+
+	using FInstanceDataType = FStateTreeMoveToTaskInstanceData;
+
+	virtual const UStruct* GetInstanceDataType() const override { return FInstanceDataType::StaticStruct(); }
+
+	GAMEPLAYSTATETREEMODULE_API virtual EStateTreeRunStatus EnterState(FStateTreeExecutionContext& Context, const FStateTreeTransitionResult& Transition) const override;
+	GAMEPLAYSTATETREEMODULE_API virtual EStateTreeRunStatus Tick(FStateTreeExecutionContext& Context, const float DeltaTime) const override;
+	GAMEPLAYSTATETREEMODULE_API virtual void ExitState(FStateTreeExecutionContext& Context, const FStateTreeTransitionResult& Transition) const override;
+
+	GAMEPLAYSTATETREEMODULE_API virtual UAITask_MoveTo* PrepareMoveToTask(FStateTreeExecutionContext& Context, AAIController& Controller, UAITask_MoveTo* ExistingTask, FAIMoveRequest& MoveRequest) const;
+	GAMEPLAYSTATETREEMODULE_API virtual EStateTreeRunStatus PerformMoveTask(FStateTreeExecutionContext& Context, AAIController& Controller) const;
+};
\ No newline at end of file
