Create IStateTreeSchemaProvider that allows a derived class to override the schema used by a parent property.
Create StateTreeAIComponent & StateTreeAIComponentSchema derived from StateTreeComponent but guaranteed to have a AIController

diff --git a/Engine/Plugins/Runtime/GameplayStateTree/Source/GameplayStateTreeModule/Private/BehaviorTree/Tasks/BTTask_RunStateTree.cpp b/Engine/Plugins/Runtime/GameplayStateTree/Source/GameplayStateTreeModule/Private/BehaviorTree/Tasks/BTTask_RunStateTree.cpp
--- a/Engine/Plugins/Runtime/GameplayStateTree/Source/GameplayStateTreeModule/Private/BehaviorTree/Tasks/BTTask_RunStateTree.cpp
+++ b/Engine/Plugins/Runtime/GameplayStateTree/Source/GameplayStateTreeModule/Private/BehaviorTree/Tasks/BTTask_RunStateTree.cpp
@@ -3,7 +3,7 @@
 #include "BehaviorTree/Tasks/BTTask_RunStateTree.h"
 
 #include "BehaviorTree/GameplayStateTreeBTUtils.h"
-#include "Components/StateTreeComponentSchema.h"
+#include "Components/StateTreeAIComponentSchema.h"
 #include "StateTreeExecutionContext.h"
 
 UBTTask_RunStateTree::UBTTask_RunStateTree(const FObjectInitializer& ObjectInitializer /*= FObjectInitializer::Get()*/)
@@ -20,7 +20,7 @@ EBTNodeResult::Type UBTTask_RunStateTree::ExecuteTask(UBehaviorTreeComponent& Ow
 	if (StateTreeRef.IsValid())
 	{
 		FStateTreeExecutionContext Context(*OwnerComp.GetOwner(), *StateTreeRef.GetStateTree(), InstanceData);
-		if (UStateTreeComponentSchema::SetContextRequirements(OwnerComp, Context))
+		if (SetContextRequirements(OwnerComp, Context))
 		{
 			const EStateTreeRunStatus StartStatus = Context.Start(&StateTreeRef.GetParameters());
 			return GameplayStateTreeBTUtils::StateTreeRunStatusToBTNodeResult(StartStatus);
@@ -35,7 +35,7 @@ void UBTTask_RunStateTree::TickTask(UBehaviorTreeComponent& OwnerComp, uint8* No
 	if (StateTreeRef.IsValid())
 	{
 		FStateTreeExecutionContext Context(*OwnerComp.GetOwner(), *StateTreeRef.GetStateTree(), InstanceData);
-		if (UStateTreeComponentSchema::SetContextRequirements(OwnerComp, Context))
+		if (SetContextRequirements(OwnerComp, Context))
 		{
 			const EStateTreeRunStatus TickStatus = Context.Tick(DeltaSeconds);
 			if (TickStatus != EStateTreeRunStatus::Running)
@@ -58,9 +58,25 @@ void UBTTask_RunStateTree::OnTaskFinished(UBehaviorTreeComponent& OwnerComp, uin
 	if (StateTreeRef.IsValid())
 	{
 		FStateTreeExecutionContext Context(*OwnerComp.GetOwner(), *StateTreeRef.GetStateTree(), InstanceData);
-		if(UStateTreeComponentSchema::SetContextRequirements(OwnerComp, Context))
+		if(SetContextRequirements(OwnerComp, Context))
 		{
 			Context.Stop();
 		}
 	}
 }
+
+bool UBTTask_RunStateTree::SetContextRequirements(UBehaviorTreeComponent& OwnerComp, FStateTreeExecutionContext& Context)
+{
+	Context.SetCollectExternalDataCallback(FOnCollectStateTreeExternalData::CreateUObject(this, &UBTTask_RunStateTree::CollectExternalData));
+	return UStateTreeAIComponentSchema::SetContextRequirements(OwnerComp, Context);
+}
+
+bool UBTTask_RunStateTree::CollectExternalData(const FStateTreeExecutionContext& Context, const UStateTree* StateTree, TArrayView<const FStateTreeExternalDataDesc> ExternalDataDescs, TArrayView<FStateTreeDataView> OutDataViews)
+{
+	return UStateTreeAIComponentSchema::CollectExternalData(Context, StateTree, ExternalDataDescs, OutDataViews);
+}
+
+TSubclassOf<UStateTreeSchema> UBTTask_RunStateTree::GetSchema() const
+{
+	return UStateTreeAIComponentSchema::StaticClass();
+}
diff --git a/Engine/Plugins/Runtime/GameplayStateTree/Source/GameplayStateTreeModule/Private/Components/StateTreeAIComponent.cpp b/Engine/Plugins/Runtime/GameplayStateTree/Source/GameplayStateTreeModule/Private/Components/StateTreeAIComponent.cpp
new file mode 100644
--- /dev/null
+++ b/Engine/Plugins/Runtime/GameplayStateTree/Source/GameplayStateTreeModule/Private/Components/StateTreeAIComponent.cpp
@@ -0,0 +1,17 @@
+// Copyright Epic Games, Inc. All Rights Reserved.
+
+#include "Components/StateTreeAIComponent.h"
+
+#include "Components/StateTreeAIComponentSchema.h"
+#include "StateTreeExecutionContext.h"
+
+TSubclassOf<UStateTreeSchema> UStateTreeAIComponent::GetSchema() const
+{
+	return UStateTreeAIComponentSchema::StaticClass();
+}
+
+bool UStateTreeAIComponent::SetContextRequirements(FStateTreeExecutionContext& Context, bool bLogErrors /*= false*/)
+{
+	Context.SetCollectExternalDataCallback(FOnCollectStateTreeExternalData::CreateUObject(this, &UStateTreeAIComponent::CollectExternalData));
+	return UStateTreeAIComponentSchema::SetContextRequirements(*this, Context, bLogErrors);
+}
diff --git a/Engine/Plugins/Runtime/GameplayStateTree/Source/GameplayStateTreeModule/Private/Components/StateTreeAIComponentSchema.cpp b/Engine/Plugins/Runtime/GameplayStateTree/Source/GameplayStateTreeModule/Private/Components/StateTreeAIComponentSchema.cpp
new file mode 100644
--- /dev/null
+++ b/Engine/Plugins/Runtime/GameplayStateTree/Source/GameplayStateTreeModule/Private/Components/StateTreeAIComponentSchema.cpp
@@ -0,0 +1,54 @@
+// Copyright Epic Games, Inc. All Rights Reserved.
+
+#include "Components/StateTreeAIComponentSchema.h"
+
+#include "AIController.h"
+#include "BrainComponent.h"
+#include "StateTreeExecutionContext.h"
+#include "Tasks/StateTreeAITask.h"
+
+UStateTreeAIComponentSchema::UStateTreeAIComponentSchema(const FObjectInitializer& ObjectInitializer /*= FObjectInitializer::Get()*/)
+	: AIControllerClass(AAIController::StaticClass())
+{
+	ContextDataDescs.Emplace(TEXT("AIController"), AIControllerClass.Get(), FGuid(0xEDB3CD97, 0x95F94E0A, 0xBD15207B, 0x98645CDC));
+}
+
+void UStateTreeAIComponentSchema::PostLoad()
+{
+	Super::PostLoad();
+	ContextDataDescs[1].Struct = AIControllerClass.Get();
+}
+
+bool UStateTreeAIComponentSchema::IsStructAllowed(const UScriptStruct* InScriptStruct) const
+{
+	return Super::IsStructAllowed(InScriptStruct) || InScriptStruct->IsChildOf(FStateTreeAITaskBase::StaticStruct());
+}
+
+bool UStateTreeAIComponentSchema::SetContextRequirements(UBrainComponent& BrainComponent, FStateTreeExecutionContext& Context, bool bLogErrors /*= false*/)
+{
+	if(!Context.IsValid())
+	{
+		return false;
+	}
+
+	const FName AIControllerName(TEXT("AIController"));
+	Context.SetContextDataByName(AIControllerName, FStateTreeDataView(BrainComponent.GetAIOwner()));
+
+	return Super::SetContextRequirements(BrainComponent, Context, bLogErrors);
+}
+
+#if WITH_EDITOR
+void UStateTreeAIComponentSchema::PostEditChangeChainProperty(FPropertyChangedChainEvent& PropertyChangedEvent)
+{
+	Super::PostEditChangeChainProperty(PropertyChangedEvent);
+
+	if (FProperty* Property = PropertyChangedEvent.Property)
+	{
+		if (Property->GetOwnerClass() == UStateTreeAIComponentSchema::StaticClass()
+			&& Property->GetFName() == GET_MEMBER_NAME_CHECKED(UStateTreeAIComponentSchema, AIControllerClass))
+		{
+			ContextDataDescs[1].Struct = AIControllerClass.Get();
+		}
+	}
+}
+#endif // WITH_EDITOR
diff --git a/Engine/Plugins/Runtime/GameplayStateTree/Source/GameplayStateTreeModule/Private/Components/StateTreeComponent.cpp b/Engine/Plugins/Runtime/GameplayStateTree/Source/GameplayStateTreeModule/Private/Components/StateTreeComponent.cpp
--- a/Engine/Plugins/Runtime/GameplayStateTree/Source/GameplayStateTreeModule/Private/Components/StateTreeComponent.cpp
+++ b/Engine/Plugins/Runtime/GameplayStateTree/Source/GameplayStateTreeModule/Private/Components/StateTreeComponent.cpp
@@ -65,6 +65,7 @@ bool UStateTreeComponent::CollectExternalData(const FStateTreeExecutionContext&
 
 bool UStateTreeComponent::SetContextRequirements(FStateTreeExecutionContext& Context, bool bLogErrors)
 {
+	Context.SetCollectExternalDataCallback(FOnCollectStateTreeExternalData::CreateUObject(this, &UStateTreeComponent::CollectExternalData));
 	return UStateTreeComponentSchema::SetContextRequirements(*this, Context);
 }
 
@@ -101,7 +102,7 @@ void UStateTreeComponent::TickComponent(float DeltaTime, enum ELevelTick TickTyp
 	}
 
 	FStateTreeExecutionContext Context(*GetOwner(), *StateTreeRef.GetStateTree(), InstanceData);
-	if (UStateTreeComponentSchema::SetContextRequirements(*this, Context))
+	if (SetContextRequirements(Context))
 	{
 		const EStateTreeRunStatus PreviousRunStatus = Context.GetStateTreeRunStatus();
 		const EStateTreeRunStatus CurrentRunStatus = Context.Tick(DeltaTime);
@@ -124,7 +125,7 @@ void UStateTreeComponent::StartLogic()
 	}
 
 	FStateTreeExecutionContext Context(*GetOwner(), *StateTreeRef.GetStateTree(), InstanceData);
-	if (UStateTreeComponentSchema::SetContextRequirements(*this, Context))
+	if (SetContextRequirements(Context))
 	{
 		const EStateTreeRunStatus PreviousRunStatus = Context.GetStateTreeRunStatus();
 		const EStateTreeRunStatus CurrentRunStatus = Context.Start(&StateTreeRef.GetParameters());
@@ -148,7 +149,7 @@ void UStateTreeComponent::RestartLogic()
 	}
 
 	FStateTreeExecutionContext Context(*GetOwner(), *StateTreeRef.GetStateTree(), InstanceData);
-	if (UStateTreeComponentSchema::SetContextRequirements(*this, Context))
+	if (SetContextRequirements(Context))
 	{
 		const EStateTreeRunStatus PreviousRunStatus = Context.GetStateTreeRunStatus();
 		const EStateTreeRunStatus CurrentRunStatus = Context.Start(&StateTreeRef.GetParameters());
@@ -177,7 +178,7 @@ void UStateTreeComponent::StopLogic(const FString& Reason)
 	}
 
 	FStateTreeExecutionContext Context(*GetOwner(), *StateTreeRef.GetStateTree(), InstanceData);
-	if (UStateTreeComponentSchema::SetContextRequirements(*this, Context))
+	if (SetContextRequirements(Context))
 	{
 		const EStateTreeRunStatus PreviousRunStatus = Context.GetStateTreeRunStatus();
 		const EStateTreeRunStatus CurrentRunStatus = Context.Stop();
@@ -291,6 +292,11 @@ void UStateTreeComponent::OnGameplayTaskInitialized(UGameplayTask& Task)
 	}
 }
 
+TSubclassOf<UStateTreeSchema> UStateTreeComponent::GetSchema() const
+{
+	return UStateTreeComponentSchema::StaticClass();
+}
+
 void UStateTreeComponent::SetStartLogicAutomatically(const bool bInStartLogicAutomatically)
 {
 	bStartLogicAutomatically = bInStartLogicAutomatically;
diff --git a/Engine/Plugins/Runtime/GameplayStateTree/Source/GameplayStateTreeModule/Private/Components/StateTreeComponentSchema.cpp b/Engine/Plugins/Runtime/GameplayStateTree/Source/GameplayStateTreeModule/Private/Components/StateTreeComponentSchema.cpp
--- a/Engine/Plugins/Runtime/GameplayStateTree/Source/GameplayStateTreeModule/Private/Components/StateTreeComponentSchema.cpp
+++ b/Engine/Plugins/Runtime/GameplayStateTree/Source/GameplayStateTreeModule/Private/Components/StateTreeComponentSchema.cpp
@@ -10,11 +10,12 @@
 #include "StateTreeExecutionContext.h"
 #include "StateTreeTaskBase.h"
 #include "Subsystems/WorldSubsystem.h"
+#include "Tasks/StateTreeAITask.h"
 #include "VisualLogger/VisualLogger.h"
 
 UStateTreeComponentSchema::UStateTreeComponentSchema()
 	: ContextActorClass(AActor::StaticClass())
-	, ContextActorDataDesc(FName("Actor"), AActor::StaticClass(), FGuid(0x1D971B00, 0x28884FDE, 0xB5436802, 0x36984FD5))
+	, ContextDataDescs({{ FName("Actor"), AActor::StaticClass(), FGuid(0x1D971B00, 0x28884FDE, 0xB5436802, 0x36984FD5) }})
 {
 }
 
@@ -39,13 +40,13 @@ bool UStateTreeComponentSchema::IsExternalItemAllowed(const UStruct& InStruct) c
 
 TConstArrayView<FStateTreeExternalDataDesc> UStateTreeComponentSchema::GetContextDataDescs() const
 {
-	return MakeArrayView(&ContextActorDataDesc, 1);
+	return ContextDataDescs;
 }
 
 void UStateTreeComponentSchema::PostLoad()
 {
 	Super::PostLoad();
-	ContextActorDataDesc.Struct = ContextActorClass.Get();
+	GetContextActorDataDesc().Struct = ContextActorClass.Get();
 }
 
 #if WITH_EDITOR
@@ -60,7 +61,7 @@ void UStateTreeComponentSchema::PostEditChangeChainProperty(FPropertyChangedChai
 		if (Property->GetOwnerClass() == UStateTreeComponentSchema::StaticClass()
 			&& Property->GetFName() == GET_MEMBER_NAME_CHECKED(UStateTreeComponentSchema, ContextActorClass))
 		{
-			ContextActorDataDesc.Struct = ContextActorClass.Get();
+			GetContextActorDataDesc().Struct = ContextActorClass.Get();
 		}
 	}
 }
@@ -73,8 +74,6 @@ bool UStateTreeComponentSchema::SetContextRequirements(UBrainComponent& BrainCom
 		return false;
 	}
 
-	Context.SetCollectExternalDataCallback(FOnCollectStateTreeExternalData::CreateStatic(&UStateTreeComponentSchema::CollectExternalData));
-
 	// Make sure the actor matches one required.
 	AActor* ContextActor = nullptr;
 	const UStateTreeComponentSchema* Schema = Cast<UStateTreeComponentSchema>(Context.GetStateTree()->GetSchema());
diff --git a/Engine/Plugins/Runtime/GameplayStateTree/Source/GameplayStateTreeModule/Public/BehaviorTree/Tasks/BTTask_RunStateTree.h b/Engine/Plugins/Runtime/GameplayStateTree/Source/GameplayStateTreeModule/Public/BehaviorTree/Tasks/BTTask_RunStateTree.h
--- a/Engine/Plugins/Runtime/GameplayStateTree/Source/GameplayStateTreeModule/Public/BehaviorTree/Tasks/BTTask_RunStateTree.h
+++ b/Engine/Plugins/Runtime/GameplayStateTree/Source/GameplayStateTreeModule/Public/BehaviorTree/Tasks/BTTask_RunStateTree.h
@@ -3,16 +3,19 @@
 #pragma once
 
 #include "BehaviorTree/BTTaskNode.h"
+#include "IStateTreeSchemaProvider.h"
 #include "StateTreeReference.h"
 #include "StateTreeInstanceData.h"
 
 #include "BTTask_RunStateTree.generated.h"
 
+struct FStateTreeExecutionContext;
+
 /**
- * RunStateTree task allows the execution of state tree with the StateTreeComponentSchema inside a behavior tree.
+ * RunStateTree task allows the execution of state tree with the StateTreeAIComponentSchema inside a behavior tree.
  */
  UCLASS(MinimalAPI)
-class UBTTask_RunStateTree : public UBTTaskNode
+class UBTTask_RunStateTree : public UBTTaskNode, public IStateTreeSchemaProvider
 {
 	GENERATED_BODY()
 public:
@@ -23,8 +26,13 @@ protected:
 	GAMEPLAYSTATETREEMODULE_API virtual void TickTask(UBehaviorTreeComponent& OwnerComp, uint8* NodeMemory, float DeltaSeconds) override;
 	GAMEPLAYSTATETREEMODULE_API virtual void OnTaskFinished(UBehaviorTreeComponent& OwnerComp, uint8* NodeMemory, EBTNodeResult::Type TaskResult) override;
 	
+	virtual bool SetContextRequirements(UBehaviorTreeComponent& OwnerComp, FStateTreeExecutionContext& Context);
+	virtual bool CollectExternalData(const FStateTreeExecutionContext& Context, const UStateTree* StateTree, TArrayView<const FStateTreeExternalDataDesc> ExternalDataDescs, TArrayView<FStateTreeDataView> OutDataViews);
+
+	virtual TSubclassOf<UStateTreeSchema> GetSchema() const override;
+
 	/** State tree that will be run when the task is selected. */
-	UPROPERTY(EditAnywhere, Category = Task, meta = (Schema = "/Script/GameplayStateTreeModule.StateTreeComponentSchema"))
+	UPROPERTY(EditAnywhere, Category = Task, meta = (Schema = "/Script/GameplayStateTreeModule.StateTreeAIComponentSchema"))
 	FStateTreeReference StateTreeRef;
 
 	UPROPERTY(Transient)
@@ -37,4 +45,4 @@ protected:
 	/** Random deviation on the interval between each state tree update. */
 	UPROPERTY(EditAnywhere, Category = Task, meta = (ClampMin = "0.0"))
 	float RandomDeviation = 0.0f;
-};
+ };
diff --git a/Engine/Plugins/Runtime/GameplayStateTree/Source/GameplayStateTreeModule/Public/Components/StateTreeAIComponent.h b/Engine/Plugins/Runtime/GameplayStateTree/Source/GameplayStateTreeModule/Public/Components/StateTreeAIComponent.h
new file mode 100644
--- /dev/null
+++ b/Engine/Plugins/Runtime/GameplayStateTree/Source/GameplayStateTreeModule/Public/Components/StateTreeAIComponent.h
@@ -0,0 +1,23 @@
+// Copyright Epic Games, Inc. All Rights Reserved.
+
+#pragma once
+
+#include "Components/StateTreeComponent.h"
+
+#include "StateTreeAIComponent.generated.h"
+
+/**
+* State tree component designed to be run on an AIController.
+* It uses the StateTreeAIComponentSchema that guarantees access to the AIController.
+*/
+UCLASS(ClassGroup = AI, meta = (BlueprintSpawnableComponent))
+class UStateTreeAIComponent : public UStateTreeComponent
+{
+	GENERATED_BODY()
+public:
+	// BEGIN IStateTreeSchemaProvider
+	TSubclassOf<UStateTreeSchema> GetSchema() const override;
+	// END
+
+	virtual bool SetContextRequirements(FStateTreeExecutionContext& Context, bool bLogErrors = false) override;
+};
diff --git a/Engine/Plugins/Runtime/GameplayStateTree/Source/GameplayStateTreeModule/Public/Components/StateTreeAIComponentSchema.h b/Engine/Plugins/Runtime/GameplayStateTree/Source/GameplayStateTreeModule/Public/Components/StateTreeAIComponentSchema.h
new file mode 100644
--- /dev/null
+++ b/Engine/Plugins/Runtime/GameplayStateTree/Source/GameplayStateTreeModule/Public/Components/StateTreeAIComponentSchema.h
@@ -0,0 +1,36 @@
+// Copyright Epic Games, Inc. All Rights Reserved.
+
+#pragma once
+
+#include "Components/StateTreeComponentSchema.h"
+
+#include "StateTreeAIComponentSchema.generated.h"
+
+class AAIController;
+
+/**
+* State tree schema to be used with StateTreeAIComponent. 
+* It guarantees access to an AIController and the Actor context value can be used to access the controlled pawn.
+*/
+UCLASS(BlueprintType, EditInlineNew, CollapseCategories, meta = (DisplayName = "StateTree AI Component", CommonSchema))
+class UStateTreeAIComponentSchema : public UStateTreeComponentSchema
+{
+	GENERATED_BODY()
+public:
+	UStateTreeAIComponentSchema(const FObjectInitializer& ObjectInitializer = FObjectInitializer::Get());
+	
+	virtual void PostLoad() override;
+
+	virtual bool IsStructAllowed(const UScriptStruct* InScriptStruct) const override;
+
+	static bool SetContextRequirements(UBrainComponent& BrainComponent, FStateTreeExecutionContext& Context, bool bLogErrors = false);
+
+#if WITH_EDITOR
+	virtual void PostEditChangeChainProperty(FPropertyChangedChainEvent& PropertyChangedEvent) override;
+#endif // WITH_EDITOR
+
+protected:
+	/** AIController class the StateTree is expected to run on. Allows to bind to specific Actor class' properties. */
+	UPROPERTY(EditAnywhere, Category = "Defaults", NoClear)
+	TSubclassOf<AAIController> AIControllerClass = nullptr;
+};
diff --git a/Engine/Plugins/Runtime/GameplayStateTree/Source/GameplayStateTreeModule/Public/Components/StateTreeComponent.h b/Engine/Plugins/Runtime/GameplayStateTree/Source/GameplayStateTreeModule/Public/Components/StateTreeComponent.h
--- a/Engine/Plugins/Runtime/GameplayStateTree/Source/GameplayStateTreeModule/Public/Components/StateTreeComponent.h
+++ b/Engine/Plugins/Runtime/GameplayStateTree/Source/GameplayStateTreeModule/Public/Components/StateTreeComponent.h
@@ -4,6 +4,7 @@
 
 #include "BrainComponent.h"
 #include "GameplayTaskOwnerInterface.h"
+#include "IStateTreeSchemaProvider.h"
 #include "StateTreeReference.h"
 #include "StateTreeInstanceData.h"
 #include "UObject/Package.h"
@@ -19,7 +20,7 @@ class UStateTree;
 DECLARE_DYNAMIC_MULTICAST_DELEGATE_OneParam(FStateTreeRunStatusChanged, EStateTreeRunStatus, StateTreeRunStatus);
 
 UCLASS(ClassGroup = AI, HideCategories = (Activation, Collision), meta = (BlueprintSpawnableComponent))
-class GAMEPLAYSTATETREEMODULE_API UStateTreeComponent : public UBrainComponent, public IGameplayTaskOwnerInterface
+class GAMEPLAYSTATETREEMODULE_API UStateTreeComponent : public UBrainComponent, public IGameplayTaskOwnerInterface, public IStateTreeSchemaProvider
 {
 	GENERATED_BODY()
 public:
@@ -52,6 +53,10 @@ public:
 	virtual void OnGameplayTaskInitialized(UGameplayTask& Task) override;
 	// END IGameplayTaskOwnerInterface
 
+	// BEGIN IStateTreeSchemaProvider
+	TSubclassOf<UStateTreeSchema> GetSchema() const override;
+	// END
+
 	/**
 	 * Sets whether the State Tree is started automatically on being play.
 	 * This function sets the bStartLogicAutomatically property, and should be used mostly from constructions sscripts.
@@ -85,11 +90,9 @@ protected:
 	virtual void PostLoad() override;
 #endif
 	
-	UE_DEPRECATED(5.4, "Please use UStateTreeComponentSchema::SetContextRequirements instead.")
-	bool SetContextRequirements(FStateTreeExecutionContext& Context, bool bLogErrors = false);
+	virtual bool SetContextRequirements(FStateTreeExecutionContext& Context, bool bLogErrors = false);
 	
-	UE_DEPRECATED(5.4, "Please use UStateTreeComponentSchema::CollectExternalData instead.")
-	bool CollectExternalData(const FStateTreeExecutionContext& Context, const UStateTree* StateTree, TArrayView<const FStateTreeExternalDataDesc> Descs, TArrayView<FStateTreeDataView> OutDataViews) const;
+	virtual bool CollectExternalData(const FStateTreeExecutionContext& Context, const UStateTree* StateTree, TArrayView<const FStateTreeExternalDataDesc> Descs, TArrayView<FStateTreeDataView> OutDataViews) const;
 	
 #if WITH_EDITORONLY_DATA
 	UE_DEPRECATED(5.1, "This property has been deprecated. Use StateTreeReference instead.")
diff --git a/Engine/Plugins/Runtime/GameplayStateTree/Source/GameplayStateTreeModule/Public/Components/StateTreeComponentSchema.h b/Engine/Plugins/Runtime/GameplayStateTree/Source/GameplayStateTreeModule/Public/Components/StateTreeComponentSchema.h
--- a/Engine/Plugins/Runtime/GameplayStateTree/Source/GameplayStateTreeModule/Public/Components/StateTreeComponentSchema.h
+++ b/Engine/Plugins/Runtime/GameplayStateTree/Source/GameplayStateTreeModule/Public/Components/StateTreeComponentSchema.h
@@ -41,10 +41,17 @@ protected:
 	virtual void PostEditChangeChainProperty(FPropertyChangedChainEvent& PropertyChangedEvent) override;
 #endif // WITH_EDITOR
 	
+	const FStateTreeExternalDataDesc& GetContextActorDataDesc() const { return ContextDataDescs[0]; }
+	FStateTreeExternalDataDesc& GetContextActorDataDesc() { return ContextDataDescs[0]; }
+
 	/** Actor class the StateTree is expected to run on. Allows to bind to specific Actor class' properties. */
-	UPROPERTY(EditAnywhere, Category="Defaults")
+	UPROPERTY(EditAnywhere, Category="Defaults", NoClear)
 	TSubclassOf<AActor> ContextActorClass;
 	
+	UE_DEPRECATED(5.4, "ContextActorDataDesc is being replaced with ContextDataDescs. Call GetContextActorDataDesc to access the equivalent.")
 	UPROPERTY()
 	FStateTreeExternalDataDesc ContextActorDataDesc;
+
+	UPROPERTY()
+	TArray<FStateTreeExternalDataDesc> ContextDataDescs;
 };
diff --git a/Engine/Plugins/Runtime/StateTree/Source/StateTreeEditorModule/Private/Customizations/StateTreeReferenceDetails.cpp b/Engine/Plugins/Runtime/StateTree/Source/StateTreeEditorModule/Private/Customizations/StateTreeReferenceDetails.cpp
--- a/Engine/Plugins/Runtime/StateTree/Source/StateTreeEditorModule/Private/Customizations/StateTreeReferenceDetails.cpp
+++ b/Engine/Plugins/Runtime/StateTree/Source/StateTreeEditorModule/Private/Customizations/StateTreeReferenceDetails.cpp
@@ -6,6 +6,7 @@
 #include "StateTreeDelegates.h"
 #include "IDetailChildrenBuilder.h"
 #include "IPropertyUtilities.h"
+#include "IStateTreeSchemaProvider.h"
 #include "PropertyBagDetails.h"
 #include "PropertyCustomizationHelpers.h"
 #include "StateTreeReference.h"
@@ -116,7 +117,7 @@ void FStateTreeReferenceDetails::CustomizeHeader(const TSharedRef<IPropertyHandl
 	const TSharedPtr<IPropertyHandle> StateTreeProperty = StructPropertyHandle->GetChildHandle(GET_MEMBER_NAME_CHECKED(FStateTreeReference, StateTree));
 	check(StateTreeProperty.IsValid());
 
-	const FString SchemaMetaDataValue = InStructPropertyHandle->GetMetaData(UE::StateTree::SchemaTag);
+	const FString SchemaMetaDataValue = GetSchemaPath(*StructPropertyHandle);
 
 	InHeaderRow
 	.NameContent()
@@ -199,4 +200,23 @@ void FStateTreeReferenceDetails::SyncParameters(const UStateTree* StateTreeToSyn
 	}
 }
 
+FString FStateTreeReferenceDetails::GetSchemaPath(IPropertyHandle& StructPropertyHandle)
+{
+	TArray<UObject*> OuterObjects;
+	StructPropertyHandle.GetOuterObjects(OuterObjects);
+
+	for (const UObject* OuterObject : OuterObjects)
+	{
+		if (const IStateTreeSchemaProvider* SchemaProvider = Cast<IStateTreeSchemaProvider>(OuterObject))
+		{
+			if (const UClass* SchemaClass = SchemaProvider->GetSchema().Get())
+			{
+				return SchemaClass->GetPathName();
+			}
+		}
+	}
+
+	return StructPropertyHandle.GetMetaData(UE::StateTree::SchemaTag);
+}
+
 #undef LOCTEXT_NAMESPACE
diff --git a/Engine/Plugins/Runtime/StateTree/Source/StateTreeEditorModule/Private/Customizations/StateTreeReferenceDetails.h b/Engine/Plugins/Runtime/StateTree/Source/StateTreeEditorModule/Private/Customizations/StateTreeReferenceDetails.h
--- a/Engine/Plugins/Runtime/StateTree/Source/StateTreeEditorModule/Private/Customizations/StateTreeReferenceDetails.h
+++ b/Engine/Plugins/Runtime/StateTree/Source/StateTreeEditorModule/Private/Customizations/StateTreeReferenceDetails.h
@@ -34,6 +34,8 @@ private:
 	 */
 	void SyncParameters(const UStateTree* StateTreeToSync = nullptr) const;
 
+	static FString GetSchemaPath(IPropertyHandle& StructPropertyHandle);
+
 	TSharedPtr<IPropertyHandle> StructPropertyHandle;
 	TSharedPtr<IPropertyHandle> ParametersHandle;
 	TSharedPtr<IPropertyUtilities> PropUtils;
diff --git a/Engine/Plugins/Runtime/StateTree/Source/StateTreeModule/Public/IStateTreeSchemaProvider.h b/Engine/Plugins/Runtime/StateTree/Source/StateTreeModule/Public/IStateTreeSchemaProvider.h
new file mode 100644
--- /dev/null
+++ b/Engine/Plugins/Runtime/StateTree/Source/StateTreeModule/Public/IStateTreeSchemaProvider.h
@@ -0,0 +1,26 @@
+// Copyright Epic Games, Inc. All Rights Reserved.
+
+#pragma once
+
+class UStateTreeSchema;
+
+#include "UObject/Interface.h"
+
+#include "IStateTreeSchemaProvider.generated.h"
+
+UINTERFACE(MinimalAPI)
+class UStateTreeSchemaProvider : public UInterface
+{
+	GENERATED_BODY()
+};
+
+/**
+* Implementing this interface allows derived class to override the schema used to filter valid state trees for a FStateTreeReference.
+*/
+class IStateTreeSchemaProvider
+{
+	GENERATED_BODY()
+
+public:
+	virtual TSubclassOf<UStateTreeSchema> GetSchema() const PURE_VIRTUAL(IStateTreeSchemaProvider::GetSchema(), return nullptr;)
+};
\ No newline at end of file
