FIx CIS errors
#rnx

diff --git a/Engine/Plugins/Runtime/GameplayStateTree/Source/GameplayStateTreeModule/Private/Tasks/StateTreeMoveToTask.cpp b/Engine/Plugins/Runtime/GameplayStateTree/Source/GameplayStateTreeModule/Private/Tasks/StateTreeMoveToTask.cpp
--- a/Engine/Plugins/Runtime/GameplayStateTree/Source/GameplayStateTreeModule/Private/Tasks/StateTreeMoveToTask.cpp
+++ b/Engine/Plugins/Runtime/GameplayStateTree/Source/GameplayStateTreeModule/Private/Tasks/StateTreeMoveToTask.cpp
@@ -107,7 +107,8 @@ EStateTreeRunStatus FStateTreeMoveToTask::PerformMoveTask(FStateTreeExecutionCon
 
 	if (MoveReq.IsValid())
 	{	
-		if (InstanceData.MoveToTask = PrepareMoveToTask(Context, Controller, InstanceData.MoveToTask, MoveReq))
+		InstanceData.MoveToTask = PrepareMoveToTask(Context, Controller, InstanceData.MoveToTask, MoveReq);
+		if (InstanceData.MoveToTask)
 		{
 			if (InstanceData.MoveToTask->IsActive())
 			{
diff --git a/Engine/Plugins/Runtime/StateTree/Source/StateTreeModule/Public/IStateTreeSchemaProvider.h b/Engine/Plugins/Runtime/StateTree/Source/StateTreeModule/Public/IStateTreeSchemaProvider.h
--- a/Engine/Plugins/Runtime/StateTree/Source/StateTreeModule/Public/IStateTreeSchemaProvider.h
+++ b/Engine/Plugins/Runtime/StateTree/Source/StateTreeModule/Public/IStateTreeSchemaProvider.h
@@ -2,9 +2,10 @@
 
 #pragma once
 
-class UStateTreeSchema;
-
 #include "UObject/Interface.h"
+#include "Templates/SubclassOf.h"
+
+class UStateTreeSchema;
 
 #include "IStateTreeSchemaProvider.generated.h"
 
