Create StateTreeRunEnvQueyTask for running env query in a state tree
#rnx

diff --git a/Engine/Plugins/Runtime/GameplayStateTree/Source/GameplayStateTreeModule/Private/Tasks/StateTreeRunEnvQueryTask.cpp b/Engine/Plugins/Runtime/GameplayStateTree/Source/GameplayStateTreeModule/Private/Tasks/StateTreeRunEnvQueryTask.cpp
new file mode 100644
--- /dev/null
+++ b/Engine/Plugins/Runtime/GameplayStateTree/Source/GameplayStateTreeModule/Private/Tasks/StateTreeRunEnvQueryTask.cpp
@@ -0,0 +1,117 @@
+// Copyright Epic Games, Inc. All Rights Reserved.
+
+#include "Tasks/StateTreeRunEnvQueryTask.h"
+
+#include "EnvironmentQuery/EnvQuery.h"
+#include "EnvironmentQuery/EnvQueryManager.h"
+#include "StateTreeExecutionContext.h"
+
+EStateTreeRunStatus FStateTreeRunEnvQueryTask::EnterState(FStateTreeExecutionContext& Context, const FStateTreeTransitionResult& Transition) const
+{
+	FInstanceDataType& InstanceData = Context.GetInstanceData(*this);
+	if (!InstanceData.QueryTemplate)
+	{
+		return EStateTreeRunStatus::Failed;
+	}
+
+	UObject* QueryOwner = InstanceData.OptionalQueryOwner ? InstanceData.OptionalQueryOwner.Get() : Context.GetOwner();
+	FEnvQueryRequest Request(InstanceData.QueryTemplate, QueryOwner);
+
+	for (FAIDynamicParam& DynamicParam : InstanceData.QueryConfig)
+	{
+		Request.SetDynamicParam(DynamicParam, nullptr);
+	}
+
+	InstanceData.RequestId = Request.Execute(InstanceData.RunMode,
+		FQueryFinishedSignature::CreateLambda([InstanceDataRef = Context.GetInstanceDataStructRef(*this)](TSharedPtr<FEnvQueryResult> QueryResult) mutable
+			{
+				if (FInstanceDataType* InstanceData = InstanceDataRef.GetPtr())
+				{
+					InstanceData->QueryResult = QueryResult;
+					InstanceData->RequestId = INDEX_NONE;
+				}
+			}));
+	return InstanceData.RequestId != INDEX_NONE ? EStateTreeRunStatus::Running : EStateTreeRunStatus::Failed;
+}
+
+EStateTreeRunStatus FStateTreeRunEnvQueryTask::Tick(FStateTreeExecutionContext& Context, const float DeltaTime) const
+{
+	FInstanceDataType& InstanceData = Context.GetInstanceData(*this);
+	if (InstanceData.QueryResult)
+	{
+		if (InstanceData.QueryResult->IsSuccessful())
+		{
+			auto [VectorPtr, ActorPtr, ArrayOfVector, ArrayOfActor] = InstanceData.Result.GetMutablePtrTuple<FVector, AActor*, TArray<FVector>, TArray<AActor*>>(Context);
+			if (VectorPtr)
+			{
+				*VectorPtr = InstanceData.QueryResult->GetItemAsLocation(0);
+			}
+			else if (ActorPtr)
+			{
+				*ActorPtr = InstanceData.QueryResult->GetItemAsActor(0);
+			}
+			else if (ArrayOfVector)
+			{
+				InstanceData.QueryResult->GetAllAsLocations(*ArrayOfVector);
+			}
+			else if (ArrayOfActor)
+			{
+				InstanceData.QueryResult->GetAllAsActors(*ArrayOfActor);
+			}
+			return EStateTreeRunStatus::Succeeded;
+		}
+		else
+		{
+			return EStateTreeRunStatus::Failed;
+		}
+	}
+	return EStateTreeRunStatus::Running;
+}
+
+void FStateTreeRunEnvQueryTask::ExitState(FStateTreeExecutionContext& Context, const FStateTreeTransitionResult& Transition) const
+{
+	FInstanceDataType& InstanceData = Context.GetInstanceData(*this);
+	if (InstanceData.RequestId != INDEX_NONE)
+	{
+		if (UEnvQueryManager* QueryManager = UEnvQueryManager::GetCurrent(Context.GetOwner()))
+		{
+			QueryManager->AbortQuery(InstanceData.RequestId);
+		}
+		InstanceData.RequestId = INDEX_NONE;
+	}
+	InstanceData.QueryResult.Reset();
+}
+
+#if WITH_EDITOR
+void FStateTreeRunEnvQueryTask::PostEditInstanceDataChangeChainProperty(const FPropertyChangedChainEvent& PropertyChangedEvent, FStateTreeDataView InstanceDataView)
+{
+	if (PropertyChangedEvent.GetPropertyName() == GET_MEMBER_NAME_CHECKED(FStateTreeRunEnvQueryInstanceData, QueryTemplate))
+	{
+		FInstanceDataType& InstanceData = InstanceDataView.GetMutable<FInstanceDataType>();
+		if (InstanceData.QueryTemplate)
+		{
+			InstanceData.QueryTemplate->CollectQueryParams(*InstanceData.QueryTemplate, InstanceData.QueryConfig);
+			for (FAIDynamicParam& DynamicParam : InstanceData.QueryConfig)
+			{
+				DynamicParam.bAllowBBKey = false;
+			}
+		}
+		else
+		{
+			InstanceData.QueryConfig.Reset();
+		}
+	}
+	else if (PropertyChangedEvent.Property->GetFName() == GET_MEMBER_NAME_CHECKED(FAIDynamicParam, bAllowBBKey))
+	{
+		FInstanceDataType& InstanceData = InstanceDataView.GetMutable<FInstanceDataType>();
+		const int32 ChangedIndex = PropertyChangedEvent.GetArrayIndex(GET_MEMBER_NAME_CHECKED(FStateTreeRunEnvQueryInstanceData, QueryConfig).ToString());
+		if (InstanceData.QueryConfig.IsValidIndex(ChangedIndex))
+		{
+			if (!InstanceData.QueryConfig[ChangedIndex].bAllowBBKey)
+			{
+				InstanceData.QueryConfig[ChangedIndex].BBKey.InvalidateResolvedKey();
+			}
+		}
+	}
+}
+#endif // WITH_EDITOR
\ No newline at end of file
diff --git a/Engine/Plugins/Runtime/GameplayStateTree/Source/GameplayStateTreeModule/Public/Tasks/StateTreeRunEnvQueryTask.h b/Engine/Plugins/Runtime/GameplayStateTree/Source/GameplayStateTreeModule/Public/Tasks/StateTreeRunEnvQueryTask.h
new file mode 100644
--- /dev/null
+++ b/Engine/Plugins/Runtime/GameplayStateTree/Source/GameplayStateTreeModule/Public/Tasks/StateTreeRunEnvQueryTask.h
@@ -0,0 +1,66 @@
+// Copyright Epic Games, Inc. All Rights Reserved.
+
+#pragma once
+
+#include "StateTreeTaskBase.h"
+
+#include "EnvironmentQuery/EnvQueryTypes.h"
+#include "NativeGameplayTags.h"
+#include "StateTreePropertyRef.h"
+
+#include "StateTreeRunEnvQueryTask.generated.h"
+
+USTRUCT()
+struct FStateTreeRunEnvQueryInstanceData
+{
+	GENERATED_BODY()
+
+	// Result of the query. If an array is binded, it will output all the created values otherwise it will output the best one.
+	UPROPERTY(EditAnywhere, Category = Out, meta = (RefType = "Vector, Actor", CanRefToArray))
+	FStateTreePropertyRef Result;
+
+	// If set, the query will be run with this object has the owner object. Otherwise the task will use the context one.
+	UPROPERTY(EditAnywhere, Category = Parameter)
+	TObjectPtr<UObject> OptionalQueryOwner = nullptr;
+
+	// The query template to run
+	UPROPERTY(EditAnywhere, Category = Parameter)
+	TObjectPtr<UEnvQuery> QueryTemplate;
+
+	// Query config associated with the query template.
+	UPROPERTY(EditAnywhere, EditFixedSize, Category = Parameter)
+	TArray<FAIDynamicParam> QueryConfig;
+
+	/** determines which item will be stored (All = only first matching) */
+	UPROPERTY(EditAnywhere, Category = Parameter)
+	TEnumAsByte<EEnvQueryRunMode::Type> RunMode;
+
+	TSharedPtr<FEnvQueryResult> QueryResult = nullptr;
+
+	int32 RequestId = INDEX_NONE;
+};
+
+/**
+* Task that runs an async environment query and outputs the result to an outside parameter. Supports Actor and vector types EQS.
+* The task is usually run in a sibling state to the result user will be with the data being stored in the parent state's parameters.
+* - Parent (Has an EQS result parameter)
+*	- Run Env Query (If success go to Use Query Result)
+*	- Use Query Result
+*/
+USTRUCT(meta = (DisplayName = "Run Env Query", Category = "Common"))
+struct FStateTreeRunEnvQueryTask : public FStateTreeTaskCommonBase
+{
+	GENERATED_BODY()
+
+	using FInstanceDataType = FStateTreeRunEnvQueryInstanceData;
+
+	virtual const UStruct* GetInstanceDataType() const override { return FInstanceDataType::StaticStruct(); }
+
+	virtual EStateTreeRunStatus EnterState(FStateTreeExecutionContext& Context, const FStateTreeTransitionResult& Transition) const override;
+	virtual EStateTreeRunStatus Tick(FStateTreeExecutionContext& Context, const float DeltaTime) const override;
+	virtual void ExitState(FStateTreeExecutionContext& Context, const FStateTreeTransitionResult& Transition) const override;
+
+#if WITH_EDITOR
+	virtual void PostEditInstanceDataChangeChainProperty(const FPropertyChangedChainEvent& PropertyChangedEvent, FStateTreeDataView InstanceDataView) override;
+#endif // WITH_EDITOR
+};
diff --git a/Engine/Source/Runtime/AIModule/Classes/EnvironmentQuery/EnvQueryTypes.h b/Engine/Source/Runtime/AIModule/Classes/EnvironmentQuery/EnvQueryTypes.h
--- a/Engine/Source/Runtime/AIModule/Classes/EnvironmentQuery/EnvQueryTypes.h
+++ b/Engine/Source/Runtime/AIModule/Classes/EnvironmentQuery/EnvQueryTypes.h
@@ -1295,17 +1295,22 @@ struct FAIDynamicParam
 	UPROPERTY(VisibleAnywhere, BlueprintReadOnly, Category = EQS)
 	EAIParamType ParamType;
 
+	UPROPERTY(BlueprintReadWrite, Category = EQS)
+	uint8 bAllowBBKey : 1;
+
 	UPROPERTY(EditAnywhere, BlueprintReadWrite, Category = EQS)
 	float Value;
 
-	UPROPERTY(EditAnywhere, BlueprintReadWrite, Category = EQS)
+	UPROPERTY(EditAnywhere, BlueprintReadWrite, Category = EQS, meta=(EditCondition="bAllowBBKey", EditConditionHides))
 	FBlackboardKeySelector BBKey;
 
+
 	FAIDynamicParam()
 	{
 		ParamType = EAIParamType::Float;
 		Value = 0.f;
 		BBKey.AllowNoneAsValue(true);
+		bAllowBBKey = true;
 	}
 
 	AIMODULE_API void ConfigureBBKey(UObject &QueryOwner);
diff --git a/Engine/Source/Runtime/AIModule/Private/EnvironmentQuery/EnvQueryManager.cpp b/Engine/Source/Runtime/AIModule/Private/EnvironmentQuery/EnvQueryManager.cpp
--- a/Engine/Source/Runtime/AIModule/Private/EnvironmentQuery/EnvQueryManager.cpp
+++ b/Engine/Source/Runtime/AIModule/Private/EnvironmentQuery/EnvQueryManager.cpp
@@ -78,9 +78,10 @@ FEnvQueryRequest& FEnvQueryRequest::SetNamedParams(const TArray<FEnvNamedValue>&
 FEnvQueryRequest& FEnvQueryRequest::SetDynamicParam(const FAIDynamicParam& Param, const UBlackboardComponent* BlackboardComponent)
 {
 	checkf(BlackboardComponent || (Param.BBKey.IsSet() == false), TEXT("BBKey.IsSet but no BlackboardComponent provided"));
+	ensureMsgf(Param.bAllowBBKey || !Param.BBKey.IsSet(), TEXT("Dynamic Param %s for query template %s doesn't allow BBKey but has a bbkey set."), *Param.ParamName.ToString(), *GetNameSafe(QueryTemplate));
 
 	// check if given param requires runtime resolve, like reading from BB
-	if (Param.BBKey.IsSet() && BlackboardComponent)
+	if (Param.bAllowBBKey && Param.BBKey.IsSet() && BlackboardComponent)
 	{
 		// grab info from BB
 		switch (Param.ParamType)
diff --git a/Engine/Source/Runtime/AIModule/Private/EnvironmentQuery/EnvQueryTypes.cpp b/Engine/Source/Runtime/AIModule/Private/EnvironmentQuery/EnvQueryTypes.cpp
--- a/Engine/Source/Runtime/AIModule/Private/EnvironmentQuery/EnvQueryTypes.cpp
+++ b/Engine/Source/Runtime/AIModule/Private/EnvironmentQuery/EnvQueryTypes.cpp
@@ -469,6 +469,17 @@ void FEQSParametrizedQueryExecutionRequest::PostEditChangeProperty(UObject& Owne
 			QueryConfig.Reset();
 		}
 	}
+	else if (PropertyChangedEvent.Property->GetFName() == GET_MEMBER_NAME_CHECKED(FAIDynamicParam, bAllowBBKey))
+	{
+		const int32 ChangedIndex = PropertyChangedEvent.GetArrayIndex(GET_MEMBER_NAME_CHECKED(FEQSParametrizedQueryExecutionRequest, QueryConfig).ToString());
+		if (QueryConfig.IsValidIndex(ChangedIndex))
+		{
+			if (!QueryConfig[ChangedIndex].bAllowBBKey)
+			{
+				QueryConfig[ChangedIndex].BBKey.InvalidateResolvedKey();
+			}
+		}
+	}
 }
 #endif // WITH_EDITOR
 
