Create State Tree AI Condition base namespace
#rnx

diff --git a/Engine/Plugins/Runtime/GameplayStateTree/Source/GameplayStateTreeModule/Private/Components/StateTreeAIComponentSchema.cpp b/Engine/Plugins/Runtime/GameplayStateTree/Source/GameplayStateTreeModule/Private/Components/StateTreeAIComponentSchema.cpp
--- a/Engine/Plugins/Runtime/GameplayStateTree/Source/GameplayStateTreeModule/Private/Components/StateTreeAIComponentSchema.cpp
+++ b/Engine/Plugins/Runtime/GameplayStateTree/Source/GameplayStateTreeModule/Private/Components/StateTreeAIComponentSchema.cpp
@@ -4,6 +4,7 @@
 
 #include "AIController.h"
 #include "BrainComponent.h"
+#include "Conditions/StateTreeAIConditionBase.h"
 #include "StateTreeExecutionContext.h"
 #include "Tasks/StateTreeAITask.h"
 
@@ -21,12 +22,14 @@ void UStateTreeAIComponentSchema::PostLoad()
 
 bool UStateTreeAIComponentSchema::IsStructAllowed(const UScriptStruct* InScriptStruct) const
 {
-	return Super::IsStructAllowed(InScriptStruct) || InScriptStruct->IsChildOf(FStateTreeAITaskBase::StaticStruct());
+	return Super::IsStructAllowed(InScriptStruct)
+		|| InScriptStruct->IsChildOf(FStateTreeAITaskBase::StaticStruct())
+		|| InScriptStruct->IsChildOf(FStateTreeAIConditionBase::StaticStruct());
 }
 
 bool UStateTreeAIComponentSchema::SetContextRequirements(UBrainComponent& BrainComponent, FStateTreeExecutionContext& Context, bool bLogErrors /*= false*/)
 {
-	if(!Context.IsValid())
+	if (!Context.IsValid())
 	{
 		return false;
 	}
diff --git a/Engine/Plugins/Runtime/GameplayStateTree/Source/GameplayStateTreeModule/Public/Conditions/StateTreeAIConditionBase.h b/Engine/Plugins/Runtime/GameplayStateTree/Source/GameplayStateTreeModule/Public/Conditions/StateTreeAIConditionBase.h
new file mode 100644
--- /dev/null
+++ b/Engine/Plugins/Runtime/GameplayStateTree/Source/GameplayStateTreeModule/Public/Conditions/StateTreeAIConditionBase.h
@@ -0,0 +1,14 @@
+// Copyright Epic Games, Inc. All Rights Reserved.
+
+#pragma once
+
+#include "StateTreeConditionBase.h"
+
+#include "StateTreeAIConditionBase.generated.h"
+
+// Base class of all AI condition that expect to be run on an AIController or derived class 
+USTRUCT(meta = (Hidden, Category = "AI"))
+struct FStateTreeAIConditionBase : public FStateTreeConditionBase
+{
+	GENERATED_BODY()
+};
\ No newline at end of file
