Make ContextActorClass of UStateTreeAIComponentSchema default to APawn so that it binds to the pawn instead of the controller by default
Export UStateTreeAIComponentSchema for dll linking
#rnx

diff --git a/Engine/Plugins/Runtime/GameplayStateTree/Source/GameplayStateTreeModule/Private/Components/StateTreeAIComponentSchema.cpp b/Engine/Plugins/Runtime/GameplayStateTree/Source/GameplayStateTreeModule/Private/Components/StateTreeAIComponentSchema.cpp
--- a/Engine/Plugins/Runtime/GameplayStateTree/Source/GameplayStateTreeModule/Private/Components/StateTreeAIComponentSchema.cpp
+++ b/Engine/Plugins/Runtime/GameplayStateTree/Source/GameplayStateTreeModule/Private/Components/StateTreeAIComponentSchema.cpp
@@ -5,12 +5,17 @@
 #include "AIController.h"
 #include "BrainComponent.h"
 #include "Conditions/StateTreeAIConditionBase.h"
+#include "Gameframework/Pawn.h"
 #include "StateTreeExecutionContext.h"
 #include "Tasks/StateTreeAITask.h"
 
 UStateTreeAIComponentSchema::UStateTreeAIComponentSchema(const FObjectInitializer& ObjectInitializer /*= FObjectInitializer::Get()*/)
 	: AIControllerClass(AAIController::StaticClass())
 {
+	check(ContextDataDescs.Num() == 1 && ContextDataDescs[0].Struct == AActor::StaticClass());
+	// Make the Actor a pawn by default so it binds to the controlled pawn instead of the AIController.
+	ContextActorClass = APawn::StaticClass();
+	ContextDataDescs[0].Struct = ContextActorClass.Get();
 	ContextDataDescs.Emplace(TEXT("AIController"), AIControllerClass.Get(), FGuid(0xEDB3CD97, 0x95F94E0A, 0xBD15207B, 0x98645CDC));
 }
 
diff --git a/Engine/Plugins/Runtime/GameplayStateTree/Source/GameplayStateTreeModule/Public/Components/StateTreeAIComponent.h b/Engine/Plugins/Runtime/GameplayStateTree/Source/GameplayStateTreeModule/Public/Components/StateTreeAIComponent.h
--- a/Engine/Plugins/Runtime/GameplayStateTree/Source/GameplayStateTreeModule/Public/Components/StateTreeAIComponent.h
+++ b/Engine/Plugins/Runtime/GameplayStateTree/Source/GameplayStateTreeModule/Public/Components/StateTreeAIComponent.h
@@ -11,7 +11,7 @@
 * It uses the StateTreeAIComponentSchema that guarantees access to the AIController.
 */
 UCLASS(ClassGroup = AI, meta = (BlueprintSpawnableComponent))
-class UStateTreeAIComponent : public UStateTreeComponent
+class GAMEPLAYSTATETREEMODULE_API UStateTreeAIComponent : public UStateTreeComponent
 {
 	GENERATED_BODY()
 public:
