Fix include case
#rnx

diff --git a/Engine/Plugins/Runtime/GameplayStateTree/Source/GameplayStateTreeModule/Private/Components/StateTreeAIComponentSchema.cpp b/Engine/Plugins/Runtime/GameplayStateTree/Source/GameplayStateTreeModule/Private/Components/StateTreeAIComponentSchema.cpp
--- a/Engine/Plugins/Runtime/GameplayStateTree/Source/GameplayStateTreeModule/Private/Components/StateTreeAIComponentSchema.cpp
+++ b/Engine/Plugins/Runtime/GameplayStateTree/Source/GameplayStateTreeModule/Private/Components/StateTreeAIComponentSchema.cpp
@@ -5,7 +5,7 @@
 #include "AIController.h"
 #include "BrainComponent.h"
 #include "Conditions/StateTreeAIConditionBase.h"
-#include "Gameframework/Pawn.h"
+#include "GameFramework/Pawn.h"
 #include "StateTreeExecutionContext.h"
 #include "Tasks/StateTreeAITask.h"
 
