Add default value to Run mode of Run Env Query State tree task

diff --git a/Engine/Plugins/Runtime/GameplayStateTree/Source/GameplayStateTreeModule/Public/Tasks/StateTreeRunEnvQueryTask.h b/Engine/Plugins/Runtime/GameplayStateTree/Source/GameplayStateTreeModule/Public/Tasks/StateTreeRunEnvQueryTask.h
--- a/Engine/Plugins/Runtime/GameplayStateTree/Source/GameplayStateTreeModule/Public/Tasks/StateTreeRunEnvQueryTask.h
+++ b/Engine/Plugins/Runtime/GameplayStateTree/Source/GameplayStateTreeModule/Public/Tasks/StateTreeRunEnvQueryTask.h
@@ -33,7 +33,7 @@ struct FStateTreeRunEnvQueryInstanceData
 
 	/** determines which item will be stored (All = only first matching) */
 	UPROPERTY(EditAnywhere, Category = Parameter)
-	TEnumAsByte<EEnvQueryRunMode::Type> RunMode;
+	TEnumAsByte<EEnvQueryRunMode::Type> RunMode = EEnvQueryRunMode::SingleResult;
 
 	TSharedPtr<FEnvQueryResult> QueryResult = nullptr;
 
