Change how owner of query in StateTreeRunEnvQuery is feed to the task
#rnx

diff --git a/Engine/Plugins/Runtime/GameplayStateTree/Source/GameplayStateTreeModule/Private/Tasks/StateTreeRunEnvQueryTask.cpp b/Engine/Plugins/Runtime/GameplayStateTree/Source/GameplayStateTreeModule/Private/Tasks/StateTreeRunEnvQueryTask.cpp
--- a/Engine/Plugins/Runtime/GameplayStateTree/Source/GameplayStateTreeModule/Private/Tasks/StateTreeRunEnvQueryTask.cpp
+++ b/Engine/Plugins/Runtime/GameplayStateTree/Source/GameplayStateTreeModule/Private/Tasks/StateTreeRunEnvQueryTask.cpp
@@ -14,8 +14,7 @@ EStateTreeRunStatus FStateTreeRunEnvQueryTask::EnterState(FStateTreeExecutionCon
 		return EStateTreeRunStatus::Failed;
 	}
 
-	UObject* QueryOwner = InstanceData.OptionalQueryOwner ? InstanceData.OptionalQueryOwner.Get() : Context.GetOwner();
-	FEnvQueryRequest Request(InstanceData.QueryTemplate, QueryOwner);
+	FEnvQueryRequest Request(InstanceData.QueryTemplate, InstanceData.QueryOwner);
 
 	for (FAIDynamicParam& DynamicParam : InstanceData.QueryConfig)
 	{
diff --git a/Engine/Plugins/Runtime/GameplayStateTree/Source/GameplayStateTreeModule/Public/Tasks/StateTreeRunEnvQueryTask.h b/Engine/Plugins/Runtime/GameplayStateTree/Source/GameplayStateTreeModule/Public/Tasks/StateTreeRunEnvQueryTask.h
--- a/Engine/Plugins/Runtime/GameplayStateTree/Source/GameplayStateTreeModule/Public/Tasks/StateTreeRunEnvQueryTask.h
+++ b/Engine/Plugins/Runtime/GameplayStateTree/Source/GameplayStateTreeModule/Public/Tasks/StateTreeRunEnvQueryTask.h
@@ -19,9 +19,9 @@ struct FStateTreeRunEnvQueryInstanceData
 	UPROPERTY(EditAnywhere, Category = Out, meta = (RefType = "Vector, Actor", CanRefToArray))
 	FStateTreePropertyRef Result;
 
-	// If set, the query will be run with this object has the owner object. Otherwise the task will use the context one.
-	UPROPERTY(EditAnywhere, Category = Parameter)
-	TObjectPtr<UObject> OptionalQueryOwner = nullptr;
+	// The query will be run with this actor has the owner object.
+	UPROPERTY(EditAnywhere, Category = Context)
+	TObjectPtr<AActor> QueryOwner = nullptr;
 
 	// The query template to run
 	UPROPERTY(EditAnywhere, Category = Parameter)
