PR #11971: Fixed issue in StateTreeAIComponent not starting automatically when attached to an AIController
Introduced with CL 21673181. It looks like the logic was inverted. The context requirement functionality got refactored after the original modification and it will not put the component in a wrong situation if it gets attached later to a valid owner.

diff --git a/Engine/Plugins/Runtime/GameplayStateTree/Source/GameplayStateTreeModule/Private/Components/StateTreeComponent.cpp b/Engine/Plugins/Runtime/GameplayStateTree/Source/GameplayStateTreeModule/Private/Components/StateTreeComponent.cpp
--- a/Engine/Plugins/Runtime/GameplayStateTree/Source/GameplayStateTreeModule/Private/Components/StateTreeComponent.cpp
+++ b/Engine/Plugins/Runtime/GameplayStateTree/Source/GameplayStateTreeModule/Private/Components/StateTreeComponent.cpp
@@ -74,7 +74,7 @@ void UStateTreeComponent::BeginPlay()
 {
 	Super::BeginPlay();
 	
-	if (AIOwner == nullptr && bStartLogicAutomatically)
+	if (bStartLogicAutomatically)
 	{
 		StartLogic();
 	}
