StateTree: Update comments for docgen.
#rnx

diff --git a/Engine/Plugins/Runtime/GameplayStateTree/Source/GameplayStateTreeModule/Public/Components/StateTreeAIComponent.h b/Engine/Plugins/Runtime/GameplayStateTree/Source/GameplayStateTreeModule/Public/Components/StateTreeAIComponent.h
--- a/Engine/Plugins/Runtime/GameplayStateTree/Source/GameplayStateTreeModule/Public/Components/StateTreeAIComponent.h
+++ b/Engine/Plugins/Runtime/GameplayStateTree/Source/GameplayStateTreeModule/Public/Components/StateTreeAIComponent.h
@@ -15,9 +15,9 @@ class GAMEPLAYSTATETREEMODULE_API UStateTreeAIComponent : public UStateTreeCompo
 {
 	GENERATED_BODY()
 public:
-	// BEGIN IStateTreeSchemaProvider
+	//~ BEGIN IStateTreeSchemaProvider
 	TSubclassOf<UStateTreeSchema> GetSchema() const override;
-	// END
+	//~ END
 
 	virtual bool SetContextRequirements(FStateTreeExecutionContext& Context, bool bLogErrors = false) override;
 };
diff --git a/Engine/Plugins/Runtime/GameplayStateTree/Source/GameplayStateTreeModule/Public/Components/StateTreeComponent.h b/Engine/Plugins/Runtime/GameplayStateTree/Source/GameplayStateTreeModule/Public/Components/StateTreeComponent.h
--- a/Engine/Plugins/Runtime/GameplayStateTree/Source/GameplayStateTreeModule/Public/Components/StateTreeComponent.h
+++ b/Engine/Plugins/Runtime/GameplayStateTree/Source/GameplayStateTreeModule/Public/Components/StateTreeComponent.h
@@ -26,15 +26,15 @@ class GAMEPLAYSTATETREEMODULE_API UStateTreeComponent : public UBrainComponent,
 public:
 	UStateTreeComponent(const FObjectInitializer& ObjectInitializer = FObjectInitializer::Get());
 
-	// BEGIN UActorComponent overrides
+	//~ BEGIN UActorComponent overrides
 	virtual void InitializeComponent() override;
 	virtual void UninitializeComponent() override;
 	virtual void BeginPlay() override;
 	virtual void EndPlay(const EEndPlayReason::Type EndPlayReason) override;
 	virtual void TickComponent(float DeltaTime, enum ELevelTick TickType, FActorComponentTickFunction *ThisTickFunction) override;
-	// END UActorComponent overrides
+	//~ END UActorComponent overrides
 
-	// BEGIN UBrainComponent overrides
+	//~ BEGIN UBrainComponent overrides
 	virtual void StartLogic() override;
 	virtual void RestartLogic() override;
 	virtual void StopLogic(const FString& Reason)  override;
@@ -43,19 +43,19 @@ public:
 	virtual EAILogicResuming::Type ResumeLogic(const FString& Reason)  override;
 	virtual bool IsRunning() const override;
 	virtual bool IsPaused() const override;
-	// END UBrainComponent overrides
+	//~ END UBrainComponent overrides
 
-	// BEGIN IGameplayTaskOwnerInterface
+	//~ BEGIN IGameplayTaskOwnerInterface
 	virtual UGameplayTasksComponent* GetGameplayTasksComponent(const UGameplayTask& Task) const override;
 	virtual AActor* GetGameplayTaskOwner(const UGameplayTask* Task) const override;
 	virtual AActor* GetGameplayTaskAvatar(const UGameplayTask* Task) const override;
 	virtual uint8 GetGameplayTaskDefaultPriority() const override;
 	virtual void OnGameplayTaskInitialized(UGameplayTask& Task) override;
-	// END IGameplayTaskOwnerInterface
+	//~ END IGameplayTaskOwnerInterface
 
-	// BEGIN IStateTreeSchemaProvider
+	//~ BEGIN IStateTreeSchemaProvider
 	TSubclassOf<UStateTreeSchema> GetSchema() const override;
-	// END
+	//~ END
 
 	/**
 	 * Sets whether the State Tree is started automatically on being play.
diff --git a/Engine/Plugins/Runtime/StateTree/Source/StateTreeModule/Public/StateTreeExecutionContext.h b/Engine/Plugins/Runtime/StateTree/Source/StateTreeModule/Public/StateTreeExecutionContext.h
--- a/Engine/Plugins/Runtime/StateTree/Source/StateTreeModule/Public/StateTreeExecutionContext.h
+++ b/Engine/Plugins/Runtime/StateTree/Source/StateTreeModule/Public/StateTreeExecutionContext.h
@@ -122,7 +122,7 @@ public:
 	/** @return const references to the instance data in use, or nullptr if the context is not valid. */
 	const FStateTreeInstanceData* GetInstanceData() const { return &InstanceData; }
 
-	/** @retuen mutable references to the instance data in use, or nullptr if the context is not valid. */
+	/** @return mutable references to the instance data in use, or nullptr if the context is not valid. */
 	FStateTreeInstanceData* GetMutableInstanceData() const { return &InstanceData; }
 
 	/** @return mutable references to the instance data in use. */
diff --git a/Engine/Plugins/Runtime/StateTree/Source/StateTreeModule/Public/StateTreeInstanceData.h b/Engine/Plugins/Runtime/StateTree/Source/StateTreeModule/Public/StateTreeInstanceData.h
--- a/Engine/Plugins/Runtime/StateTree/Source/StateTreeModule/Public/StateTreeInstanceData.h
+++ b/Engine/Plugins/Runtime/StateTree/Source/StateTreeModule/Public/StateTreeInstanceData.h
@@ -185,13 +185,13 @@ struct STATETREEMODULE_API FStateTreeInstanceStorage
 		return Wrapper.InstanceObject;
 	}
 
-	/** @return reference to StateTree execution state, or null if the instance data is not initialized. */
+	/** @return reference to StateTree execution state. */
 	const FStateTreeExecutionState& GetExecutionState() const
 	{
 		return ExecutionState;
 	}
 
-	/** @return reference to StateTree execution state, or null if the instance data is not initialized. */
+	/** @return reference to StateTree execution state. */
 	FStateTreeExecutionState& GetMutableExecutionState()
 	{
 		return ExecutionState;
