StateTree: Fix for StateTreeComponent's internal run state.
PR #10319

diff --git a/Engine/Plugins/Runtime/GameplayStateTree/Source/GameplayStateTreeModule/Private/Components/StateTreeComponent.cpp b/Engine/Plugins/Runtime/GameplayStateTree/Source/GameplayStateTreeModule/Private/Components/StateTreeComponent.cpp
--- a/Engine/Plugins/Runtime/GameplayStateTree/Source/GameplayStateTreeModule/Private/Components/StateTreeComponent.cpp
+++ b/Engine/Plugins/Runtime/GameplayStateTree/Source/GameplayStateTreeModule/Private/Components/StateTreeComponent.cpp
@@ -130,7 +130,7 @@ void UStateTreeComponent::StartLogic()
 	{
 		const EStateTreeRunStatus PreviousRunStatus = Context.GetStateTreeRunStatus();
 		const EStateTreeRunStatus CurrentRunStatus = Context.Start(&StateTreeRef.GetParameters());
-		bIsRunning = true;
+		bIsRunning = CurrentRunStatus == EStateTreeRunStatus::Running;
 		
 		if (CurrentRunStatus != PreviousRunStatus)
 		{
@@ -154,7 +154,7 @@ void UStateTreeComponent::RestartLogic()
 	{
 		const EStateTreeRunStatus PreviousRunStatus = Context.GetStateTreeRunStatus();
 		const EStateTreeRunStatus CurrentRunStatus = Context.Start(&StateTreeRef.GetParameters());
-		bIsRunning = true;
+		bIsRunning = CurrentRunStatus == EStateTreeRunStatus::Running;
 		
 		if (CurrentRunStatus != PreviousRunStatus)
 		{
