StateTree: Make StateTree Schema SetContextData virtual instead of forcing the child component class to manually call a static function.

diff --git a/Engine/Plugins/Runtime/GameplayStateTree/Source/GameplayStateTreeModule/Private/Components/StateTreeAIComponent.cpp b/Engine/Plugins/Runtime/GameplayStateTree/Source/GameplayStateTreeModule/Private/Components/StateTreeAIComponent.cpp
--- a/Engine/Plugins/Runtime/GameplayStateTree/Source/GameplayStateTreeModule/Private/Components/StateTreeAIComponent.cpp
+++ b/Engine/Plugins/Runtime/GameplayStateTree/Source/GameplayStateTreeModule/Private/Components/StateTreeAIComponent.cpp
@@ -9,9 +9,3 @@ TSubclassOf<UStateTreeSchema> UStateTreeAIComponent::GetSchema() const
 {
 	return UStateTreeAIComponentSchema::StaticClass();
 }
-
-bool UStateTreeAIComponent::SetContextRequirements(FStateTreeExecutionContext& Context, bool bLogErrors /*= false*/)
-{
-	Context.SetCollectExternalDataCallback(FOnCollectStateTreeExternalData::CreateUObject(this, &UStateTreeAIComponent::CollectExternalData));
-	return UStateTreeAIComponentSchema::SetContextRequirements(*this, Context, bLogErrors);
-}
diff --git a/Engine/Plugins/Runtime/GameplayStateTree/Source/GameplayStateTreeModule/Private/Components/StateTreeAIComponentSchema.cpp b/Engine/Plugins/Runtime/GameplayStateTree/Source/GameplayStateTreeModule/Private/Components/StateTreeAIComponentSchema.cpp
--- a/Engine/Plugins/Runtime/GameplayStateTree/Source/GameplayStateTreeModule/Private/Components/StateTreeAIComponentSchema.cpp
+++ b/Engine/Plugins/Runtime/GameplayStateTree/Source/GameplayStateTreeModule/Private/Components/StateTreeAIComponentSchema.cpp
@@ -9,6 +9,11 @@
 #include "StateTreeExecutionContext.h"
 #include "Tasks/StateTreeAITask.h"
 
+namespace UE::GameplayStateTree::Private
+{
+	static FLazyName Name_AIController = "AIController";
+}
+
 UStateTreeAIComponentSchema::UStateTreeAIComponentSchema(const FObjectInitializer& ObjectInitializer /*= FObjectInitializer::Get()*/)
 	: AIControllerClass(AAIController::StaticClass())
 {
@@ -16,7 +21,7 @@ UStateTreeAIComponentSchema::UStateTreeAIComponentSchema(const FObjectInitialize
 	// Make the Actor a pawn by default so it binds to the controlled pawn instead of the AIController.
 	ContextActorClass = APawn::StaticClass();
 	ContextDataDescs[0].Struct = ContextActorClass.Get();
-	ContextDataDescs.Emplace(TEXT("AIController"), AIControllerClass.Get(), FGuid(0xEDB3CD97, 0x95F94E0A, 0xBD15207B, 0x98645CDC));
+	ContextDataDescs.Emplace(UE::GameplayStateTree::Private::Name_AIController, AIControllerClass.Get(), FGuid(0xEDB3CD97, 0x95F94E0A, 0xBD15207B, 0x98645CDC));
 }
 
 void UStateTreeAIComponentSchema::PostLoad()
@@ -32,17 +37,12 @@ bool UStateTreeAIComponentSchema::IsStructAllowed(const UScriptStruct* InScriptS
 		|| InScriptStruct->IsChildOf(FStateTreeAIConditionBase::StaticStruct());
 }
 
-bool UStateTreeAIComponentSchema::SetContextRequirements(UBrainComponent& BrainComponent, FStateTreeExecutionContext& Context, bool bLogErrors /*= false*/)
+void UStateTreeAIComponentSchema::SetContextData(FContextDataSetter& ContextDataSetter, bool bLogErrors) const
 {
-	if (!Context.IsValid())
-	{
-		return false;
-	}
-
-	const FName AIControllerName(TEXT("AIController"));
-	Context.SetContextDataByName(AIControllerName, FStateTreeDataView(BrainComponent.GetAIOwner()));
+	const FName AIControllerName = UE::GameplayStateTree::Private::Name_AIController;
+	ContextDataSetter.SetContextDataByName(AIControllerName, FStateTreeDataView(ContextDataSetter.GetComponent()->GetAIOwner()));
 
-	return Super::SetContextRequirements(BrainComponent, Context, bLogErrors);
+	Super::SetContextData(ContextDataSetter, bLogErrors);
 }
 
 #if WITH_EDITOR
diff --git a/Engine/Plugins/Runtime/GameplayStateTree/Source/GameplayStateTreeModule/Private/Components/StateTreeComponentSchema.cpp b/Engine/Plugins/Runtime/GameplayStateTree/Source/GameplayStateTreeModule/Private/Components/StateTreeComponentSchema.cpp
--- a/Engine/Plugins/Runtime/GameplayStateTree/Source/GameplayStateTreeModule/Private/Components/StateTreeComponentSchema.cpp
+++ b/Engine/Plugins/Runtime/GameplayStateTree/Source/GameplayStateTreeModule/Private/Components/StateTreeComponentSchema.cpp
@@ -15,9 +15,14 @@
 #include "VisualLogger/VisualLogger.h"
 #include "StateTreePropertyFunctionBase.h"
 
+namespace UE::GameplayStateTree::Private
+{
+	static FLazyName Name_Actor = "Actor";
+}
+
 UStateTreeComponentSchema::UStateTreeComponentSchema()
 	: ContextActorClass(AActor::StaticClass())
-	, ContextDataDescs({{ FName("Actor"), AActor::StaticClass(), FGuid(0x1D971B00, 0x28884FDE, 0xB5436802, 0x36984FD5) }})
+	, ContextDataDescs({{ UE::GameplayStateTree::Private::Name_Actor, AActor::StaticClass(), FGuid(0x1D971B00, 0x28884FDE, 0xB5436802, 0x36984FD5) }})
 {
 }
 
@@ -50,7 +55,7 @@ TConstArrayView<FStateTreeExternalDataDesc> UStateTreeComponentSchema::GetContex
 void UStateTreeComponentSchema::PostLoad()
 {
 	Super::PostLoad();
-	GetContextActorDataDesc().Struct = ContextActorClass.Get();
+	GetContextActorDataDesc().Struct = ContextActorClass.Get() ? ContextActorClass.Get() : AActor::StaticClass();
 }
 
 #if WITH_EDITOR
@@ -65,26 +70,41 @@ void UStateTreeComponentSchema::PostEditChangeChainProperty(FPropertyChangedChai
 		if (Property->GetOwnerClass() == UStateTreeComponentSchema::StaticClass()
 			&& Property->GetFName() == GET_MEMBER_NAME_CHECKED(UStateTreeComponentSchema, ContextActorClass))
 		{
-			GetContextActorDataDesc().Struct = ContextActorClass.Get();
+			GetContextActorDataDesc().Struct = ContextActorClass.Get() ? ContextActorClass.Get() : AActor::StaticClass();
 		}
 	}
 }
 #endif
 
-bool UStateTreeComponentSchema::SetContextRequirements(UBrainComponent& BrainComponent, FStateTreeExecutionContext& Context, bool bLogErrors /*= false*/)
+UStateTreeComponentSchema::FContextDataSetter::FContextDataSetter(gsl::not_null<const UBrainComponent*> InBrainComponent, FStateTreeExecutionContext& InContext)
+	: BrainComponent(InBrainComponent)
+	, ExecutionContext(InContext)
+{}
+
+bool UStateTreeComponentSchema::FContextDataSetter::SetContextDataByName(FName Name, FStateTreeDataView DataView)
 {
-	if (!Context.IsValid())
-	{
-		return false;
-	}
+	return ExecutionContext.IsValid() ? ExecutionContext.SetContextDataByName(Name, DataView) : false;
+}
+
+gsl::not_null<const UStateTree*> UStateTreeComponentSchema::FContextDataSetter::GetStateTree() const
+{
+	return ExecutionContext.GetStateTree();
+}
 
+gsl::not_null<const UStateTreeComponentSchema*> UStateTreeComponentSchema::FContextDataSetter::GetSchema() const
+{
+	return Cast<UStateTreeComponentSchema>(ExecutionContext.GetStateTree()->GetSchema());
+}
+
+void UStateTreeComponentSchema::SetContextData(FContextDataSetter& ContextDataSetter, bool bLogErrors) const
+{
 	// Make sure the actor matches one required.
 	AActor* ContextActor = nullptr;
-	const UStateTreeComponentSchema* Schema = Cast<UStateTreeComponentSchema>(Context.GetStateTree()->GetSchema());
+	const UStateTreeComponentSchema* Schema = ContextDataSetter.GetSchema();
 	if (Schema)
 	{
-		AAIController* AIOwner = BrainComponent.GetAIOwner();
-		if (AAIController* OwnerController = (AIOwner != nullptr) ? AIOwner : Cast<AAIController>(BrainComponent.GetOwner()))
+		AAIController* AIOwner = ContextDataSetter.GetComponent()->GetAIOwner();
+		if (AAIController* OwnerController = (AIOwner != nullptr) ? AIOwner : Cast<AAIController>(ContextDataSetter.GetComponent()->GetOwner()))
 		{
 			if (OwnerController && OwnerController->IsA(Schema->GetContextActorClass()))
 			{
@@ -93,7 +113,7 @@ bool UStateTreeComponentSchema::SetContextRequirements(UBrainComponent& BrainCom
 		}
 		if (ContextActor == nullptr)
 		{
-			if (AActor* OwnerActor = (AIOwner != nullptr) ? AIOwner->GetPawn() : BrainComponent.GetOwner())
+			if (AActor* OwnerActor = (AIOwner != nullptr) ? AIOwner->GetPawn() : ContextDataSetter.GetComponent()->GetOwner())
 			{
 				if (OwnerActor && OwnerActor->IsA(Schema->GetContextActorClass()))
 				{
@@ -103,16 +123,27 @@ bool UStateTreeComponentSchema::SetContextRequirements(UBrainComponent& BrainCom
 		}
 		if (ContextActor == nullptr && bLogErrors)
 		{
-			UE_VLOG(BrainComponent.GetOwner(), LogStateTree, Error, TEXT("%s: Could not find context actor of type %s. StateTree will not update."), ANSI_TO_TCHAR(__FUNCTION__), *GetNameSafe(Schema->GetContextActorClass()));
+			UE_VLOG(ContextDataSetter.GetComponent()->GetOwner(), LogStateTree, Error, TEXT("%s: Could not find context actor of type %s. StateTree will not update."), ANSI_TO_TCHAR(__FUNCTION__), *GetNameSafe(Schema->GetContextActorClass()));
 		}
 	}
 	else if (bLogErrors)
 	{
-		UE_VLOG(BrainComponent.GetOwner(), LogStateTree, Error, TEXT("%s: Expected StateTree asset to contain StateTreeComponentSchema. StateTree will not update."), ANSI_TO_TCHAR(__FUNCTION__));
+		UE_VLOG(ContextDataSetter.GetComponent()->GetOwner(), LogStateTree, Error, TEXT("%s: Expected StateTree asset to contain StateTreeComponentSchema. StateTree will not update."), ANSI_TO_TCHAR(__FUNCTION__));
+	}
+
+	const FName ActorName = UE::GameplayStateTree::Private::Name_Actor;
+	ContextDataSetter.SetContextDataByName(ActorName, FStateTreeDataView(ContextActor));
+}
+
+bool UStateTreeComponentSchema::SetContextRequirements(UBrainComponent& BrainComponent, FStateTreeExecutionContext& Context, bool bLogErrors /*= false*/)
+{
+	if (!Context.IsValid())
+	{
+		return false;
 	}
 
-	const FName ActorName(TEXT("Actor"));
-	Context.SetContextDataByName(ActorName, FStateTreeDataView(ContextActor));
+	FContextDataSetter ContextDataSetter = { gsl::make_not_null(&BrainComponent), Context };
+	ContextDataSetter.GetSchema()->SetContextData(ContextDataSetter, bLogErrors);
 
 	bool bResult = Context.AreContextDataViewsValid();
 	if (!bResult && bLogErrors)
diff --git a/Engine/Plugins/Runtime/GameplayStateTree/Source/GameplayStateTreeModule/Public/Components/StateTreeAIComponent.h b/Engine/Plugins/Runtime/GameplayStateTree/Source/GameplayStateTreeModule/Public/Components/StateTreeAIComponent.h
--- a/Engine/Plugins/Runtime/GameplayStateTree/Source/GameplayStateTreeModule/Public/Components/StateTreeAIComponent.h
+++ b/Engine/Plugins/Runtime/GameplayStateTree/Source/GameplayStateTreeModule/Public/Components/StateTreeAIComponent.h
@@ -18,6 +18,4 @@ public:
 	//~ BEGIN IStateTreeSchemaProvider
 	TSubclassOf<UStateTreeSchema> GetSchema() const override;
 	//~ END
-
-	virtual bool SetContextRequirements(FStateTreeExecutionContext& Context, bool bLogErrors = false) override;
 };
diff --git a/Engine/Plugins/Runtime/GameplayStateTree/Source/GameplayStateTreeModule/Public/Components/StateTreeAIComponentSchema.h b/Engine/Plugins/Runtime/GameplayStateTree/Source/GameplayStateTreeModule/Public/Components/StateTreeAIComponentSchema.h
--- a/Engine/Plugins/Runtime/GameplayStateTree/Source/GameplayStateTreeModule/Public/Components/StateTreeAIComponentSchema.h
+++ b/Engine/Plugins/Runtime/GameplayStateTree/Source/GameplayStateTreeModule/Public/Components/StateTreeAIComponentSchema.h
@@ -23,7 +23,7 @@ public:
 
 	virtual bool IsStructAllowed(const UScriptStruct* InScriptStruct) const override;
 
-	static bool SetContextRequirements(UBrainComponent& BrainComponent, FStateTreeExecutionContext& Context, bool bLogErrors = false);
+	virtual void SetContextData(FContextDataSetter& ContextDataSetter, bool bLogErrors) const override;
 
 #if WITH_EDITOR
 	virtual void PostEditChangeChainProperty(FPropertyChangedChainEvent& PropertyChangedEvent) override;
diff --git a/Engine/Plugins/Runtime/GameplayStateTree/Source/GameplayStateTreeModule/Public/Components/StateTreeComponentSchema.h b/Engine/Plugins/Runtime/GameplayStateTree/Source/GameplayStateTreeModule/Public/Components/StateTreeComponentSchema.h
--- a/Engine/Plugins/Runtime/GameplayStateTree/Source/GameplayStateTreeModule/Public/Components/StateTreeComponentSchema.h
+++ b/Engine/Plugins/Runtime/GameplayStateTree/Source/GameplayStateTreeModule/Public/Components/StateTreeComponentSchema.h
@@ -32,9 +32,31 @@ protected:
 	virtual bool IsStructAllowed(const UScriptStruct* InScriptStruct) const override;
 	virtual bool IsClassAllowed(const UClass* InScriptStruct) const override;
 	virtual bool IsExternalItemAllowed(const UStruct& InStruct) const override;
-	
+
+	/** Helper class to set the context data on the ExecutionContext */
+	struct FContextDataSetter
+	{
+	public:
+		FContextDataSetter(gsl::not_null<const UBrainComponent*> BrainComponent, FStateTreeExecutionContext& Context);
+
+		gsl::not_null<const UBrainComponent*> GetComponent() const
+		{
+			return BrainComponent;
+		}
+		gsl::not_null<const UStateTree*> GetStateTree() const;
+		gsl::not_null<const UStateTreeComponentSchema*> GetSchema() const;
+
+		bool SetContextDataByName(FName Name, FStateTreeDataView DataView);
+
+	private:
+		gsl::not_null<const UBrainComponent*> BrainComponent;
+		FStateTreeExecutionContext& ExecutionContext;
+	};
+	virtual void SetContextData(FContextDataSetter& ContextDataSetter, bool bLogErrors) const;
+
 	virtual TConstArrayView<FStateTreeExternalDataDesc> GetContextDataDescs() const override;
 
+
 	virtual void PostLoad() override;
 
 #if WITH_EDITOR
@@ -48,9 +70,11 @@ protected:
 	UPROPERTY(EditAnywhere, Category="Defaults", NoClear)
 	TSubclassOf<AActor> ContextActorClass;
 	
+#if WITH_EDITORONLY_DATA
 	UE_DEPRECATED(5.4, "ContextActorDataDesc is being replaced with ContextDataDescs. Call GetContextActorDataDesc to access the equivalent.")
 	UPROPERTY()
-	FStateTreeExternalDataDesc ContextActorDataDesc;
+	FStateTreeExternalDataDesc ContextActorDataDesc_DEPRECATED;
+#endif
 
 	UPROPERTY()
 	TArray<FStateTreeExternalDataDesc> ContextDataDescs;
