StateTree: Call the base function for UStateTreeComponent::InitializeComponent and UStateTreeComponent::UninitializeComponent. P4 history doesn't mentionned it but we guess that it was to skip the UBrainComponent from initializing a behavior tree.

diff --git a/Engine/Plugins/Runtime/GameplayStateTree/Source/GameplayStateTreeModule/Private/Components/StateTreeComponent.cpp b/Engine/Plugins/Runtime/GameplayStateTree/Source/GameplayStateTreeModule/Private/Components/StateTreeComponent.cpp
--- a/Engine/Plugins/Runtime/GameplayStateTree/Source/GameplayStateTreeModule/Private/Components/StateTreeComponent.cpp
+++ b/Engine/Plugins/Runtime/GameplayStateTree/Source/GameplayStateTreeModule/Private/Components/StateTreeComponent.cpp
@@ -57,6 +57,9 @@ void UStateTreeComponent::InitializeComponent()
 	{
 		STATETREE_LOG(Error, TEXT("%s: Failed to init StateTreeContext."), ANSI_TO_TCHAR(__FUNCTION__));
 	}
+
+	// Skipping UBrainComponent
+	UActorComponent::InitializeComponent();
 }
 
 #if WITH_EDITOR
@@ -77,6 +80,8 @@ void UStateTreeComponent::PostLoad()
 
 void UStateTreeComponent::UninitializeComponent()
 {
+	// Skipping UBrainComponent
+	UActorComponent::UninitializeComponent();
 }
 
 bool UStateTreeComponent::CollectExternalData(const FStateTreeExecutionContext& Context, const UStateTree* StateTree, TArrayView<const FStateTreeExternalDataDesc> ExternalDataDescs, TArrayView<FStateTreeDataView> OutDataViews) const
