StateTree: Remove ensure that triggers when the ScheduledTick feature is deactivated. Deactivate the tick by default.

diff --git a/Engine/Plugins/Runtime/GameplayStateTree/Source/GameplayStateTreeModule/Private/Components/StateTreeComponent.cpp b/Engine/Plugins/Runtime/GameplayStateTree/Source/GameplayStateTreeModule/Private/Components/StateTreeComponent.cpp
--- a/Engine/Plugins/Runtime/GameplayStateTree/Source/GameplayStateTreeModule/Private/Components/StateTreeComponent.cpp
+++ b/Engine/Plugins/Runtime/GameplayStateTree/Source/GameplayStateTreeModule/Private/Components/StateTreeComponent.cpp
@@ -40,6 +40,7 @@ UStateTreeComponent::UStateTreeComponent(const FObjectInitializer& ObjectInitial
 	: Super(ObjectInitializer)
 {
 	bWantsInitializeComponent = true;
+	PrimaryComponentTick.bStartWithTickEnabled = false;
 	bIsRunning = false;
 	bIsPaused = false;
 }
@@ -119,7 +120,8 @@ void UStateTreeComponent::TickComponent(float DeltaTime, enum ELevelTick TickTyp
 
 	if (!bIsRunning || bIsPaused)
 	{
-		ensure(!UE::GameplayStateTree::Private::bScheduledTickAllowed);
+		STATETREE_LOG(Warning, TEXT("%hs: Ticking a paused or a not running State Tree component."), __FUNCTION__);
+		DisableTick();
 		return;
 	}
 
@@ -309,23 +311,21 @@ void UStateTreeComponent::ScheduleTickFrame(const FStateTreeScheduledTick& NextT
 			}
 		}
 	}
+	else
+	{
+		DisableTick();
+	}
 }
 
 void UStateTreeComponent::ConditionalEnableTick()
 {
 	STATETREE_LOG(Log, TEXT("%s: EnabledTick manually."), ANSI_TO_TCHAR(__FUNCTION__));
-	if (bIsRunning && !bIsPaused)
-	{
-		if (!IsComponentTickEnabled())
-		{
-			ScheduleTickFrame(FStateTreeScheduledTick::MakeNextFrame());
-		}
-	}
+	ScheduleTickFrame(FStateTreeScheduledTick::MakeNextFrame());
 }
 
 void UStateTreeComponent::DisableTick()
 {
-	if (IsComponentTickEnabled() && UE::GameplayStateTree::Private::bScheduledTickAllowed)
+	if (IsComponentTickEnabled())
 	{
 		SetComponentTickEnabled(false);
 	}
