Add virtual function to decalre if the state tree should be validated during intialization
- moves super intialize above validation checks to ensure the component is initialized properly even if validation fails

#rnx

diff --git a/Engine/Plugins/Runtime/GameplayStateTree/Source/GameplayStateTreeModule/Private/Components/StateTreeComponent.cpp b/Engine/Plugins/Runtime/GameplayStateTree/Source/GameplayStateTreeModule/Private/Components/StateTreeComponent.cpp
--- a/Engine/Plugins/Runtime/GameplayStateTree/Source/GameplayStateTreeModule/Private/Components/StateTreeComponent.cpp
+++ b/Engine/Plugins/Runtime/GameplayStateTree/Source/GameplayStateTreeModule/Private/Components/StateTreeComponent.cpp
@@ -47,20 +47,9 @@ UStateTreeComponent::UStateTreeComponent(const FObjectInitializer& ObjectInitial
 
 void UStateTreeComponent::InitializeComponent()
 {
-	if (!StateTreeRef.IsValid())
-	{
-		STATETREE_LOG(Error, TEXT("%s: StateTree asset is not set, cannot initialize."), ANSI_TO_TCHAR(__FUNCTION__));
-		return;
-	}
-
-	const FStateTreeExecutionContext Context(*GetOwner(), *StateTreeRef.GetStateTree(), InstanceData);
-	if (!Context.IsValid())
-	{
-		STATETREE_LOG(Error, TEXT("%s: Failed to init StateTreeContext."), ANSI_TO_TCHAR(__FUNCTION__));
-	}
-
 	// Skipping UBrainComponent
 	UActorComponent::InitializeComponent();
+	ValidateStateTreeReference();
 }
 
 #if WITH_EDITOR
@@ -403,6 +392,21 @@ TSubclassOf<UStateTreeSchema> UStateTreeComponent::GetSchema() const
 	return UStateTreeComponentSchema::StaticClass();
 }
 
+void UStateTreeComponent::ValidateStateTreeReference()
+{
+	if (StateTreeRef.IsValid())
+	{
+		STATETREE_LOG(Error, TEXT("%s: StateTree asset is not set, cannot initialize."), ANSI_TO_TCHAR(__FUNCTION__));
+		return;
+	}
+
+	const FStateTreeExecutionContext Context(*GetOwner(), *StateTreeRef.GetStateTree(), InstanceData);
+	if (!Context.IsValid())
+	{
+		STATETREE_LOG(Error, TEXT("%s: Failed to init StateTreeContext."), ANSI_TO_TCHAR(__FUNCTION__));
+	}
+}
+
 void UStateTreeComponent::SetStartLogicAutomatically(const bool bInStartLogicAutomatically)
 {
 	bStartLogicAutomatically = bInStartLogicAutomatically;
diff --git a/Engine/Plugins/Runtime/GameplayStateTree/Source/GameplayStateTreeModule/Public/Components/StateTreeComponent.h b/Engine/Plugins/Runtime/GameplayStateTree/Source/GameplayStateTreeModule/Public/Components/StateTreeComponent.h
--- a/Engine/Plugins/Runtime/GameplayStateTree/Source/GameplayStateTreeModule/Public/Components/StateTreeComponent.h
+++ b/Engine/Plugins/Runtime/GameplayStateTree/Source/GameplayStateTreeModule/Public/Components/StateTreeComponent.h
@@ -70,6 +70,13 @@ public:
 	TSubclassOf<UStateTreeSchema> GetSchema() const override;
 	//~ END
 
+	/**
+	 * Called during initialize, will validate the state tree reference and create a context from the state tree to check it's validity
+	 * Override this function for custom state tree validation.
+	 * Note: Override without calling super if the state tree reference is dynamically set after initilization
+	 */
+	virtual void ValidateStateTreeReference();
+	
 	/**
 	 * Sets whether the State Tree is started automatically on being play.
 	 * This function sets the bStartLogicAutomatically property, and should be used mostly from constructions sscripts.
