Fix bad test introduced with CL 39081498
#rnx

diff --git a/Engine/Plugins/Runtime/GameplayStateTree/Source/GameplayStateTreeModule/Private/Components/StateTreeComponent.cpp b/Engine/Plugins/Runtime/GameplayStateTree/Source/GameplayStateTreeModule/Private/Components/StateTreeComponent.cpp
--- a/Engine/Plugins/Runtime/GameplayStateTree/Source/GameplayStateTreeModule/Private/Components/StateTreeComponent.cpp
+++ b/Engine/Plugins/Runtime/GameplayStateTree/Source/GameplayStateTreeModule/Private/Components/StateTreeComponent.cpp
@@ -394,7 +394,7 @@ TSubclassOf<UStateTreeSchema> UStateTreeComponent::GetSchema() const
 
 void UStateTreeComponent::ValidateStateTreeReference()
 {
-	if (StateTreeRef.IsValid())
+	if (!StateTreeRef.IsValid())
 	{
 		STATETREE_LOG(Error, TEXT("%s: StateTree asset is not set, cannot initialize."), ANSI_TO_TCHAR(__FUNCTION__));
 		return;
