StateTree: Add UStateTreeComponent.SetStateTree method
PR #10216

diff --git a/Engine/Plugins/Runtime/GameplayStateTree/Source/GameplayStateTreeModule/Private/Components/StateTreeComponent.cpp b/Engine/Plugins/Runtime/GameplayStateTree/Source/GameplayStateTreeModule/Private/Components/StateTreeComponent.cpp
--- a/Engine/Plugins/Runtime/GameplayStateTree/Source/GameplayStateTreeModule/Private/Components/StateTreeComponent.cpp
+++ b/Engine/Plugins/Runtime/GameplayStateTree/Source/GameplayStateTreeModule/Private/Components/StateTreeComponent.cpp
@@ -407,6 +407,13 @@ void UStateTreeComponent::ValidateStateTreeReference()
 	}
 }
 
+void UStateTreeComponent::SetStateTree(UStateTree* NewStateTree)
+{
+	StopLogic("Setting new StateTree");
+	StateTreeRef.SetStateTree(NewStateTree);
+	ValidateStateTreeReference();
+}
+
 void UStateTreeComponent::SetStartLogicAutomatically(const bool bInStartLogicAutomatically)
 {
 	bStartLogicAutomatically = bInStartLogicAutomatically;
diff --git a/Engine/Plugins/Runtime/GameplayStateTree/Source/GameplayStateTreeModule/Public/Components/StateTreeComponent.h b/Engine/Plugins/Runtime/GameplayStateTree/Source/GameplayStateTreeModule/Public/Components/StateTreeComponent.h
--- a/Engine/Plugins/Runtime/GameplayStateTree/Source/GameplayStateTreeModule/Public/Components/StateTreeComponent.h
+++ b/Engine/Plugins/Runtime/GameplayStateTree/Source/GameplayStateTreeModule/Public/Components/StateTreeComponent.h
@@ -70,16 +70,12 @@ public:
 	TSubclassOf<UStateTreeSchema> GetSchema() const override;
 	//~ END
 
-	/**
-	 * Called during initialize, will validate the state tree reference and create a context from the state tree to check it's validity
-	 * Override this function for custom state tree validation.
-	 * Note: Override without calling super if the state tree reference is dynamically set after initilization
-	 */
-	virtual void ValidateStateTreeReference();
+	/** Stops currently running StateTree and sets new StateTree reference. */
+	void SetStateTree(UStateTree* NewStateTree);
 	
 	/**
 	 * Sets whether the State Tree is started automatically on being play.
-	 * This function sets the bStartLogicAutomatically property, and should be used mostly from constructions sscripts.
+	 * This function sets the bStartLogicAutomatically property, and should be used mostly from constructions scripts.
 	 * If you wish to start the logic manually, call StartLogic(). 
 	 */
 	UFUNCTION(BlueprintCallable, Category = "Gameplay|StateTree")
@@ -116,7 +112,13 @@ protected:
 #if WITH_EDITORONLY_DATA
 	virtual void PostLoad() override;
 #endif
-	
+	/**
+	 * Called during initialize, will validate the state tree reference and create a context from the state tree to check it's validity
+	 * Override this function for custom state tree validation.
+	 * Note: Override without calling super if the state tree reference is dynamically set after initialization
+	 */
+	virtual void ValidateStateTreeReference();
+
 	virtual bool SetContextRequirements(FStateTreeExecutionContext& Context, bool bLogErrors = false);
 	
 	virtual bool CollectExternalData(const FStateTreeExecutionContext& Context, const UStateTree* StateTree, TArrayView<const FStateTreeExternalDataDesc> Descs, TArrayView<FStateTreeDataView> OutDataViews) const;
