StateTree: Quick hack to keep the schema actor when using the behavior tree task.

diff --git a/Engine/Plugins/Runtime/GameplayStateTree/Source/GameplayStateTreeModule/Private/BehaviorTree/Tasks/BTTask_RunStateTree.cpp b/Engine/Plugins/Runtime/GameplayStateTree/Source/GameplayStateTreeModule/Private/BehaviorTree/Tasks/BTTask_RunStateTree.cpp
--- a/Engine/Plugins/Runtime/GameplayStateTree/Source/GameplayStateTreeModule/Private/BehaviorTree/Tasks/BTTask_RunStateTree.cpp
+++ b/Engine/Plugins/Runtime/GameplayStateTree/Source/GameplayStateTreeModule/Private/BehaviorTree/Tasks/BTTask_RunStateTree.cpp
@@ -7,6 +7,16 @@
 #include "StateTree.h"
 #include "StateTreeExecutionContext.h"
 
+namespace UE::GameplayStateTree::Private
+{
+	const FLazyName ActorName = "Actor";
+	AActor* GetSchemeActorContext(FStateTreeExecutionContext& Context)
+	{
+		FStateTreeDataView ActorContextView = Context.GetContextDataByName(ActorName);
+		return ActorContextView.IsValid() && ActorContextView.GetStruct()->IsChildOf(AActor::StaticClass()) ? ActorContextView.GetMutablePtr<AActor>() : nullptr;
+	}
+}
+
 UBTTask_RunStateTree::UBTTask_RunStateTree(const FObjectInitializer& ObjectInitializer /*= FObjectInitializer::Get()*/)
 	: Super(ObjectInitializer)
 {
@@ -18,12 +28,18 @@ UBTTask_RunStateTree::UBTTask_RunStateTree(const FObjectInitializer& ObjectIniti
 
 EBTNodeResult::Type UBTTask_RunStateTree::ExecuteTask(UBehaviorTreeComponent& OwnerComp, uint8* NodeMemory)
 {
+	SchemaActor.Reset();
+
 	if (StateTreeRef.IsValid())
 	{
 		FStateTreeExecutionContext Context(*OwnerComp.GetOwner(), *StateTreeRef.GetStateTree(), InstanceData);
 		if (SetContextRequirements(OwnerComp, Context))
 		{
 			const EStateTreeRunStatus StartStatus = Context.Start(&StateTreeRef.GetParameters());
+			if (StartStatus == EStateTreeRunStatus::Running)
+			{
+				SchemaActor = UE::GameplayStateTree::Private::GetSchemeActorContext(Context);
+			}
 			return GameplayStateTreeBTUtils::StateTreeRunStatusToBTNodeResult(StartStatus);
 		}
 	}
@@ -38,6 +54,7 @@ void UBTTask_RunStateTree::TickTask(UBehaviorTreeComponent& OwnerComp, uint8* No
 		FStateTreeExecutionContext Context(*OwnerComp.GetOwner(), *StateTreeRef.GetStateTree(), InstanceData);
 		if (SetContextRequirements(OwnerComp, Context))
 		{
+			ensureMsgf(SchemaActor == UE::GameplayStateTree::Private::GetSchemeActorContext(Context), TEXT("The schema actor is not the same from the previous tick."));
 			const EStateTreeRunStatus TickStatus = Context.Tick(DeltaSeconds);
 			if (TickStatus != EStateTreeRunStatus::Running)
 			{
@@ -59,7 +76,13 @@ void UBTTask_RunStateTree::OnTaskFinished(UBehaviorTreeComponent& OwnerComp, uin
 	if (StateTreeRef.IsValid())
 	{
 		FStateTreeExecutionContext Context(*OwnerComp.GetOwner(), *StateTreeRef.GetStateTree(), InstanceData);
-		if(SetContextRequirements(OwnerComp, Context))
+		if(!SetContextRequirements(OwnerComp, Context))
+		{
+			// If the controller is destroyed / GCed before this task, the state tree won't be stopped properly causing issues.
+			//@todo This is a temporary fix until the schema can save the data appropriately.
+			Context.SetContextDataByName(UE::GameplayStateTree::Private::ActorName, FStateTreeDataView(SchemaActor.Get()));
+		}
+		if (ensure(Context.AreContextDataViewsValid()))
 		{
 			Context.Stop();
 		}
diff --git a/Engine/Plugins/Runtime/GameplayStateTree/Source/GameplayStateTreeModule/Public/BehaviorTree/Tasks/BTTask_RunStateTree.h b/Engine/Plugins/Runtime/GameplayStateTree/Source/GameplayStateTreeModule/Public/BehaviorTree/Tasks/BTTask_RunStateTree.h
--- a/Engine/Plugins/Runtime/GameplayStateTree/Source/GameplayStateTreeModule/Public/BehaviorTree/Tasks/BTTask_RunStateTree.h
+++ b/Engine/Plugins/Runtime/GameplayStateTree/Source/GameplayStateTreeModule/Public/BehaviorTree/Tasks/BTTask_RunStateTree.h
@@ -49,4 +49,7 @@ protected:
 	/** Random deviation on the interval between each state tree update. */
 	UPROPERTY(EditAnywhere, Category = Task, meta = (ClampMin = "0.0"))
 	float RandomDeviation = 0.0f;
+
+	/** Saved actor used by the schema requirement. */
+	TWeakObjectPtr<AActor> SchemaActor;
  };
diff --git a/Engine/Plugins/Runtime/StateTree/Source/StateTreeModule/Private/StateTreeExecutionContext.cpp b/Engine/Plugins/Runtime/StateTree/Source/StateTreeModule/Private/StateTreeExecutionContext.cpp
--- a/Engine/Plugins/Runtime/StateTree/Source/StateTreeModule/Private/StateTreeExecutionContext.cpp
+++ b/Engine/Plugins/Runtime/StateTree/Source/StateTreeModule/Private/StateTreeExecutionContext.cpp
@@ -587,6 +587,19 @@ bool FStateTreeExecutionContext::SetContextDataByName(const FName Name, FStateTr
 	return false;
 }
 
+FStateTreeDataView FStateTreeExecutionContext::GetContextDataByName(const FName Name) const
+{
+	const FStateTreeExternalDataDesc* Desc = RootStateTree.GetContextDataDescs().FindByPredicate([&Name](const FStateTreeExternalDataDesc& Desc)
+		{
+			return Desc.Name == Name;
+		});
+	if (Desc)
+	{
+		return ContextAndExternalDataViews[Desc->Handle.DataHandle.GetIndex()];
+	}
+	return FStateTreeDataView();
+}
+
 FStateTreeWeakExecutionContext FStateTreeExecutionContext::MakeWeakExecutionContext() const
 {
 	return FStateTreeWeakExecutionContext(*this);
diff --git a/Engine/Plugins/Runtime/StateTree/Source/StateTreeModule/Public/StateTreeExecutionContext.h b/Engine/Plugins/Runtime/StateTree/Source/StateTreeModule/Public/StateTreeExecutionContext.h
--- a/Engine/Plugins/Runtime/StateTree/Source/StateTreeModule/Public/StateTreeExecutionContext.h
+++ b/Engine/Plugins/Runtime/StateTree/Source/StateTreeModule/Public/StateTreeExecutionContext.h
@@ -464,6 +464,9 @@ public:
 	/** Sets the context data based on name (name is defined in the schema), returns true if data was found */
 	bool SetContextDataByName(const FName Name, FStateTreeDataView DataView);
 
+	/** @return the context data based on name (name is defined in the schema) */
+	FStateTreeDataView GetContextDataByName(const FName Name) const;
+
 	/** @return True if all context data pointers are set. */ 
 	bool AreContextDataViewsValid() const;
 
