StateTree: Expose the SetStateTree function to BP. Add a state tree reference setter. Validate if the tree has the correct schema.

diff --git a/Engine/Plugins/Runtime/GameplayStateTree/Source/GameplayStateTreeModule/Private/Components/StateTreeComponent.cpp b/Engine/Plugins/Runtime/GameplayStateTree/Source/GameplayStateTreeModule/Private/Components/StateTreeComponent.cpp
--- a/Engine/Plugins/Runtime/GameplayStateTree/Source/GameplayStateTreeModule/Private/Components/StateTreeComponent.cpp
+++ b/Engine/Plugins/Runtime/GameplayStateTree/Source/GameplayStateTreeModule/Private/Components/StateTreeComponent.cpp
@@ -158,9 +158,8 @@ void UStateTreeComponent::StartTree()
 {
 	bIsRunning = false;
 
-	if (!StateTreeRef.IsValid())
+	if (HasValidStateTreeReference().HasError())
 	{
-		STATETREE_LOG(Warning, TEXT("Trying to restart State Tree component with invalid asset."));
 		DisableTick();
 		return;
 	}
@@ -394,24 +393,69 @@ TSubclassOf<UStateTreeSchema> UStateTreeComponent::GetSchema() const
 
 void UStateTreeComponent::ValidateStateTreeReference()
 {
-	if (!StateTreeRef.IsValid())
-	{
-		STATETREE_LOG(Error, TEXT("%s: StateTree asset is not set, cannot initialize."), ANSI_TO_TCHAR(__FUNCTION__));
-		return;
-	}
+ 	TValueOrError<void, FString> ValidResult = HasValidStateTreeReference();
+ 	if (ValidResult.HasError())
+ 	{
+ 	 	STATETREE_LOG(Error, TEXT("%s: %s. Cannot initialize."), *ValidResult.GetError(), ANSI_TO_TCHAR(__FUNCTION__));
+ 	 	return;
+ 	}
 
 	const FStateTreeExecutionContext Context(*GetOwner(), *StateTreeRef.GetStateTree(), InstanceData);
 	if (!Context.IsValid())
 	{
 		STATETREE_LOG(Error, TEXT("%s: Failed to init StateTreeContext."), ANSI_TO_TCHAR(__FUNCTION__));
+		return;
 	}
 }
 
-void UStateTreeComponent::SetStateTree(UStateTree* NewStateTree)
+ TValueOrError<void, FString> UStateTreeComponent::HasValidStateTreeReference() const
+{
+	if (!StateTreeRef.IsValid())
+	{
+		return MakeError(TEXT("The State Tree asset is not set."));
+	}
+
+	if (StateTreeRef.GetStateTree()->GetSchema() == nullptr
+		|| !StateTreeRef.GetStateTree()->GetSchema()->GetClass()->IsChildOf(UStateTreeComponentSchema::StaticClass()))
+	{
+		return MakeError(TEXT("The State Tree schema is not compatible."));
+	}
+
+	return MakeValue();
+}
+
+void UStateTreeComponent::SetStateTree(UStateTree* InStateTree)
 {
 	StopLogic("Setting new StateTree");
-	StateTreeRef.SetStateTree(NewStateTree);
-	ValidateStateTreeReference();
+
+	StateTreeRef.SetStateTree(InStateTree);
+
+	if (StateTreeRef.GetStateTree() != nullptr)
+	{
+		TValueOrError<void, FString> ValidResult = HasValidStateTreeReference();
+		if (ValidResult.HasError())
+		{
+			STATETREE_LOG(Warning, TEXT("%s: %s."), *ValidResult.GetError(), ANSI_TO_TCHAR(__FUNCTION__));
+			StateTreeRef = FStateTreeReference();
+		}
+	}
+}
+
+void UStateTreeComponent::SetStateTreeReference(FStateTreeReference InStateTreeReference)
+{
+	StopLogic("Setting new StateTree");
+
+	StateTreeRef = MoveTemp(InStateTreeReference);
+
+	if (StateTreeRef.GetStateTree() != nullptr)
+	{
+		TValueOrError<void, FString> ValidResult = HasValidStateTreeReference();
+		if (ValidResult.HasError())
+		{
+			STATETREE_LOG(Warning, TEXT("%s: %s."), *ValidResult.GetError(), ANSI_TO_TCHAR(__FUNCTION__));
+			StateTreeRef = FStateTreeReference();
+		}
+	}
 }
 
 void UStateTreeComponent::SetStartLogicAutomatically(const bool bInStartLogicAutomatically)
diff --git a/Engine/Plugins/Runtime/GameplayStateTree/Source/GameplayStateTreeModule/Public/Components/StateTreeComponent.h b/Engine/Plugins/Runtime/GameplayStateTree/Source/GameplayStateTreeModule/Public/Components/StateTreeComponent.h
--- a/Engine/Plugins/Runtime/GameplayStateTree/Source/GameplayStateTreeModule/Public/Components/StateTreeComponent.h
+++ b/Engine/Plugins/Runtime/GameplayStateTree/Source/GameplayStateTreeModule/Public/Components/StateTreeComponent.h
@@ -7,6 +7,7 @@
 #include "IStateTreeSchemaProvider.h"
 #include "StateTreeReference.h"
 #include "StateTreeInstanceData.h"
+#include "Templates/ValueOrError.h"
 #include "UObject/Package.h"
 #include "StateTreeComponent.generated.h"
 
@@ -70,8 +71,13 @@ public:
 	TSubclassOf<UStateTreeSchema> GetSchema() const override;
 	//~ END
 
-	/** Stops currently running StateTree and sets new StateTree reference. */
-	void SetStateTree(UStateTree* NewStateTree);
+	/** Stops currently running StateTree and sets new state tree reference. */
+	UFUNCTION(BlueprintCallable, Category = "Gameplay|StateTree")
+	void SetStateTree(UStateTree* StateTree);
+
+	/** Stops currently running StateTree and sets new state tree reference. */
+	UFUNCTION(BlueprintCallable, Category = "Gameplay|StateTree")
+	void SetStateTreeReference(FStateTreeReference StateTreeReference);
 	
 	/**
 	 * Sets whether the State Tree is started automatically on being play.
@@ -119,6 +125,9 @@ protected:
 	 */
 	virtual void ValidateStateTreeReference();
 
+	/** @return a value if the state tree reference can be used by the component or the error why it's not a valid reference. */
+	virtual TValueOrError<void, FString> HasValidStateTreeReference() const;
+
 	virtual bool SetContextRequirements(FStateTreeExecutionContext& Context, bool bLogErrors = false);
 	
 	virtual bool CollectExternalData(const FStateTreeExecutionContext& Context, const UStateTree* StateTree, TArrayView<const FStateTreeExternalDataDesc> Descs, TArrayView<FStateTreeDataView> OutDataViews) const;
