StateTree: Log error when collecting external data fails with the component schema.

diff --git a/Engine/Plugins/Runtime/GameplayStateTree/Source/GameplayStateTreeModule/Private/Components/StateTreeComponentSchema.cpp b/Engine/Plugins/Runtime/GameplayStateTree/Source/GameplayStateTreeModule/Private/Components/StateTreeComponentSchema.cpp
--- a/Engine/Plugins/Runtime/GameplayStateTree/Source/GameplayStateTreeModule/Private/Components/StateTreeComponentSchema.cpp
+++ b/Engine/Plugins/Runtime/GameplayStateTree/Source/GameplayStateTreeModule/Private/Components/StateTreeComponentSchema.cpp
@@ -176,20 +176,25 @@ bool UStateTreeComponentSchema::SetContextRequirements(UBrainComponent& BrainCom
 
 bool UStateTreeComponentSchema::CollectExternalData(const FStateTreeExecutionContext& Context, const UStateTree* StateTree, TArrayView<const FStateTreeExternalDataDesc> ExternalDataDescs, TArrayView<FStateTreeDataView> OutDataViews)
 {
+	checkf(ExternalDataDescs.Num() == OutDataViews.Num(), TEXT("The execution context failed to fill OutDataViews with empty values."));
+
 	const UWorld* World = Context.GetWorld();
 	if (World == nullptr)
 	{
+		UE_VLOG(Context.GetOwner(), LogStateTree, Error, TEXT("%hs: External data can't find the world ('%s' using StateTree '%s')."),
+			__FUNCTION__, *GetNameSafe(Context.GetOwner()), *GetFullNameSafe(Context.GetStateTree()));
 		return false;
 	}
 
-	check(ExternalDataDescs.Num() == OutDataViews.Num());
-
 	AActor* Owner = Cast<AActor>(Context.GetOwner());
 	if (!Owner)
 	{
+		UE_VLOG(Context.GetOwner(), LogStateTree, Error, TEXT("%hs: The external data owner is invalid ('%s' using StateTree '%s')."),
+			__FUNCTION__, *GetNameSafe(Context.GetOwner()), *GetFullNameSafe(Context.GetStateTree()));
 		return false;
 	}
 
+	int32 IssuesFoundCounter = 0;
 	AAIController* AIOwner = Cast<AAIController>(Owner);
 	for (int32 Index = 0; Index < ExternalDataDescs.Num(); Index++)
 	{
@@ -200,29 +205,39 @@ bool UStateTreeComponentSchema::CollectExternalData(const FStateTreeExecutionCon
 			{
 				UWorldSubsystem* Subsystem = World->GetSubsystemBase(Cast<UClass>(const_cast<UStruct*>(ItemDesc.Struct.Get())));
 				OutDataViews[Index] = FStateTreeDataView(Subsystem);
+				UE_CVLOG(Subsystem == nullptr, Context.GetOwner(), LogStateTree, Error, TEXT("StateTree %s: Could not find required subsystem %s"), *GetNameSafe(Context.GetStateTree()), *GetNameSafe(ItemDesc.Struct));
+				IssuesFoundCounter += Subsystem != nullptr ? 0 : 1;
 			}
 			else if (ItemDesc.Struct->IsChildOf(UActorComponent::StaticClass()))
 			{
 				UActorComponent* Component = Owner->FindComponentByClass(Cast<UClass>(const_cast<UStruct*>(ItemDesc.Struct.Get())));
 				OutDataViews[Index] = FStateTreeDataView(Component);
+				UE_CVLOG(Component == nullptr, Context.GetOwner(), LogStateTree, Error, TEXT("StateTree %s: Could not find required component %s"), *GetNameSafe(Context.GetStateTree()), *GetNameSafe(ItemDesc.Struct));
+				IssuesFoundCounter += Component != nullptr ? 0 : 1;
 			}
 			else if (ItemDesc.Struct->IsChildOf(APawn::StaticClass()))
 			{
 				APawn* OwnerPawn = (AIOwner != nullptr) ? AIOwner->GetPawn() : Cast<APawn>(Owner);
+				UE_CVLOG(OwnerPawn == nullptr, Context.GetOwner(), LogStateTree, Error, TEXT("StateTree %s: Could not find required pawn %s"), *GetNameSafe(Context.GetStateTree()), *GetNameSafe(ItemDesc.Struct));
 				OutDataViews[Index] = FStateTreeDataView(OwnerPawn);
+				IssuesFoundCounter += OwnerPawn != nullptr ? 0 : 1;
 			}
 			else if (ItemDesc.Struct->IsChildOf(AAIController::StaticClass()))
 			{
 				AAIController* OwnerController = AIOwner;
+				UE_CVLOG(OwnerController == nullptr, Context.GetOwner(), LogStateTree, Error, TEXT("StateTree %s: Could not find required controller %s"), *GetNameSafe(Context.GetStateTree()), *GetNameSafe(ItemDesc.Struct));
 				OutDataViews[Index] = FStateTreeDataView(OwnerController);
+				IssuesFoundCounter += OwnerController != nullptr ? 0 : 1;
 			}
 			else if (ItemDesc.Struct->IsChildOf(AActor::StaticClass()))
 			{
 				AActor* OwnerActor = (AIOwner != nullptr) ? AIOwner->GetPawn() : Owner;
+				UE_CVLOG(OwnerActor == nullptr, Context.GetOwner(), LogStateTree, Error, TEXT("StateTree %s: Could not find required actor %s"), *GetNameSafe(Context.GetStateTree()), *GetNameSafe(ItemDesc.Struct));
 				OutDataViews[Index] = FStateTreeDataView(OwnerActor);
+				IssuesFoundCounter += OwnerActor != nullptr ? 0 : 1;
 			}
 		}
 	}
 
-	return true;
+	return IssuesFoundCounter == 0;
 }
\ No newline at end of file
