StateTree: Make the StateTreeComponent Blueprintable.

diff --git a/Engine/Plugins/Runtime/GameplayStateTree/Source/GameplayStateTreeModule/Public/Components/StateTreeComponent.h b/Engine/Plugins/Runtime/GameplayStateTree/Source/GameplayStateTreeModule/Public/Components/StateTreeComponent.h
--- a/Engine/Plugins/Runtime/GameplayStateTree/Source/GameplayStateTreeModule/Public/Components/StateTreeComponent.h
+++ b/Engine/Plugins/Runtime/GameplayStateTree/Source/GameplayStateTreeModule/Public/Components/StateTreeComponent.h
@@ -33,7 +33,7 @@ public:
 
 DECLARE_DYNAMIC_MULTICAST_DELEGATE_OneParam(FStateTreeRunStatusChanged, EStateTreeRunStatus, StateTreeRunStatus);
 
-UCLASS(ClassGroup = AI, HideCategories = (Activation, Collision), meta = (BlueprintSpawnableComponent))
+UCLASS(Blueprintable, ClassGroup = AI, HideCategories = (Activation, Collision), meta = (BlueprintSpawnableComponent))
 class GAMEPLAYSTATETREEMODULE_API UStateTreeComponent : public UBrainComponent, public IGameplayTaskOwnerInterface, public IStateTreeSchemaProvider
 {
 	GENERATED_BODY()
