[StateTree] Mark UStateTreeAIComponentSchema for export.

diff --git a/Engine/Plugins/Runtime/GameplayStateTree/Source/GameplayStateTreeModule/Public/Components/StateTreeAIComponentSchema.h b/Engine/Plugins/Runtime/GameplayStateTree/Source/GameplayStateTreeModule/Public/Components/StateTreeAIComponentSchema.h
--- a/Engine/Plugins/Runtime/GameplayStateTree/Source/GameplayStateTreeModule/Public/Components/StateTreeAIComponentSchema.h
+++ b/Engine/Plugins/Runtime/GameplayStateTree/Source/GameplayStateTreeModule/Public/Components/StateTreeAIComponentSchema.h
@@ -13,7 +13,7 @@ class AAIController;
 * It guarantees access to an AIController and the Actor context value can be used to access the controlled pawn.
 */
 UCLASS(BlueprintType, EditInlineNew, CollapseCategories, meta = (DisplayName = "StateTree AI Component", CommonSchema))
-class UStateTreeAIComponentSchema : public UStateTreeComponentSchema
+class GAMEPLAYSTATETREEMODULE_API UStateTreeAIComponentSchema : public UStateTreeComponentSchema
 {
 	GENERATED_BODY()
 public:
