StateTree: Set the default value DefaultScheduledTickAllowed to true. That enabled the scheduled tick on all state tree components by default.
#rnx

diff --git a/Engine/Plugins/Runtime/GameplayStateTree/Source/GameplayStateTreeModule/Private/Components/StateTreeComponentSchema.cpp b/Engine/Plugins/Runtime/GameplayStateTree/Source/GameplayStateTreeModule/Private/Components/StateTreeComponentSchema.cpp
--- a/Engine/Plugins/Runtime/GameplayStateTree/Source/GameplayStateTreeModule/Private/Components/StateTreeComponentSchema.cpp
+++ b/Engine/Plugins/Runtime/GameplayStateTree/Source/GameplayStateTreeModule/Private/Components/StateTreeComponentSchema.cpp
@@ -21,7 +21,7 @@ namespace UE::GameplayStateTree::Private
 
 	FAutoConsoleVariable CVarDefaultScheduledTickAllowed(
 		TEXT("StateTree.Component.DefaultScheduledTickAllowed"),
-		false,
+		true,
 		TEXT("By default, allow execution context to sleep and the tick delayed.")
 	);
 }
