StateTree: Add a function to set the LinkedStateTreeOverrides from BP.
Change the logic for SetStateTree. It will only work if the instance is not running. That is to prevent memory corruption.

diff --git a/Engine/Plugins/Runtime/GameplayStateTree/Source/GameplayStateTreeModule/Private/Components/StateTreeComponent.cpp b/Engine/Plugins/Runtime/GameplayStateTree/Source/GameplayStateTreeModule/Private/Components/StateTreeComponent.cpp
--- a/Engine/Plugins/Runtime/GameplayStateTree/Source/GameplayStateTreeModule/Private/Components/StateTreeComponent.cpp
+++ b/Engine/Plugins/Runtime/GameplayStateTree/Source/GameplayStateTreeModule/Private/Components/StateTreeComponent.cpp
@@ -144,13 +144,13 @@ void UStateTreeComponent::TickComponent(float DeltaTime, enum ELevelTick TickTyp
 
 void UStateTreeComponent::StartLogic()
 {
-	STATETREE_LOG(Log, TEXT("%s: Start Logic"), ANSI_TO_TCHAR(__FUNCTION__));
+	STATETREE_LOG(Log, TEXT("%hs: Start Logic"), __FUNCTION__);
 	StartTree();
 }
 
 void UStateTreeComponent::RestartLogic()
 {
-	STATETREE_LOG(Log, TEXT("%s: Restart Logic"), ANSI_TO_TCHAR(__FUNCTION__));
+	STATETREE_LOG(Log, TEXT("%hs: Restart Logic"), __FUNCTION__);
 	StartTree();
 }
 
@@ -192,7 +192,7 @@ void UStateTreeComponent::StartTree()
 
 void UStateTreeComponent::StopLogic(const FString& Reason)
 {
-	STATETREE_LOG(Log, TEXT("%s: Stopping, reason: \'%s\'"), ANSI_TO_TCHAR(__FUNCTION__), *Reason);
+	STATETREE_LOG(Log, TEXT("%hs: Stopping, reason: \'%s\'"), __FUNCTION__, *Reason);
 
 	if (!bIsRunning)
 	{
@@ -204,7 +204,7 @@ void UStateTreeComponent::StopLogic(const FString& Reason)
 
 	if (!StateTreeRef.IsValid())
 	{
-		STATETREE_LOG(Warning, TEXT("%s: Trying to stop State Tree component with invalid asset."), ANSI_TO_TCHAR(__FUNCTION__));
+		STATETREE_LOG(Warning, TEXT("%hs: Trying to stop State Tree component with invalid asset."), __FUNCTION__);
 		return;
 	}
 
@@ -229,7 +229,7 @@ void UStateTreeComponent::Cleanup()
 
 void UStateTreeComponent::PauseLogic(const FString& Reason)
 {
-	STATETREE_LOG(Log, TEXT("%s: Execution updates: PAUSED (%s)"), ANSI_TO_TCHAR(__FUNCTION__), *Reason);
+	STATETREE_LOG(Log, TEXT("%hs: Execution updates: PAUSED (%s)"), __FUNCTION__, *Reason);
 	bIsPaused = true;
 
 	DisableTick();
@@ -237,7 +237,7 @@ void UStateTreeComponent::PauseLogic(const FString& Reason)
 
 EAILogicResuming::Type UStateTreeComponent::ResumeLogic(const FString& Reason)
 {
-	STATETREE_LOG(Log, TEXT("%s: Execution updates: RESUMED (%s)"), ANSI_TO_TCHAR(__FUNCTION__), *Reason);
+	STATETREE_LOG(Log, TEXT("%hs: Execution updates: RESUMED (%s)"), __FUNCTION__, *Reason);
 
 	const EAILogicResuming::Type SuperResumeResult = Super::ResumeLogic(Reason);
 
@@ -307,7 +307,7 @@ void UStateTreeComponent::ScheduleTickFrame(const FStateTreeScheduledTick& NextT
 
 void UStateTreeComponent::ConditionalEnableTick()
 {
-	STATETREE_LOG(Log, TEXT("%s: EnabledTick manually."), ANSI_TO_TCHAR(__FUNCTION__));
+	STATETREE_LOG(Log, TEXT("%hs: EnabledTick manually."), __FUNCTION__);
 	ScheduleTickFrame(FStateTreeScheduledTick::MakeNextFrame());
 }
 
@@ -396,14 +396,14 @@ void UStateTreeComponent::ValidateStateTreeReference()
  	TValueOrError<void, FString> ValidResult = HasValidStateTreeReference();
  	if (ValidResult.HasError())
  	{
- 	 	STATETREE_LOG(Error, TEXT("%s: %s. Cannot initialize."), *ValidResult.GetError(), ANSI_TO_TCHAR(__FUNCTION__));
+ 	 	STATETREE_LOG(Error, TEXT("%hs: %s. Cannot initialize."), __FUNCTION__, * ValidResult.GetError());
  	 	return;
  	}
 
 	const FStateTreeExecutionContext Context(*GetOwner(), *StateTreeRef.GetStateTree(), InstanceData);
 	if (!Context.IsValid())
 	{
-		STATETREE_LOG(Error, TEXT("%s: Failed to init StateTreeContext."), ANSI_TO_TCHAR(__FUNCTION__));
+		STATETREE_LOG(Error, TEXT("%hs: Failed to init StateTreeContext."), __FUNCTION__);
 		return;
 	}
 }
@@ -426,7 +426,15 @@ void UStateTreeComponent::ValidateStateTreeReference()
 
 void UStateTreeComponent::SetStateTree(UStateTree* InStateTree)
 {
-	StopLogic("Setting new StateTree");
+	{
+		// Can't change the StateTreeRef on a running tree. It might change the instance data while running tasks.
+		FStateTreeReadOnlyExecutionContext Context(GetOwner(), StateTreeRef.GetStateTree(), InstanceData);
+		if (Context.GetStateTreeRunStatus() == EStateTreeRunStatus::Running)
+		{
+			STATETREE_LOG(Warning, TEXT("%hs : Trying to change the state tree on a running instance."), __FUNCTION__);
+			return;
+		}
+	}
 
 	StateTreeRef.SetStateTree(InStateTree);
 
@@ -435,7 +443,7 @@ void UStateTreeComponent::SetStateTree(UStateTree* InStateTree)
 		TValueOrError<void, FString> ValidResult = HasValidStateTreeReference();
 		if (ValidResult.HasError())
 		{
-			STATETREE_LOG(Warning, TEXT("%s: %s."), *ValidResult.GetError(), ANSI_TO_TCHAR(__FUNCTION__));
+			STATETREE_LOG(Warning, TEXT("%hs: %s."), __FUNCTION__, *ValidResult.GetError());
 			StateTreeRef = FStateTreeReference();
 		}
 	}
@@ -443,7 +451,15 @@ void UStateTreeComponent::SetStateTree(UStateTree* InStateTree)
 
 void UStateTreeComponent::SetStateTreeReference(FStateTreeReference InStateTreeReference)
 {
-	StopLogic("Setting new StateTree");
+	{
+		// Can't change the StateTreeRef on a running tree. It might change the instance data while running tasks.
+		FStateTreeReadOnlyExecutionContext Context(GetOwner(), StateTreeRef.GetStateTree(), InstanceData);
+		if (Context.GetStateTreeRunStatus() == EStateTreeRunStatus::Running)
+		{
+			STATETREE_LOG(Warning, TEXT("%hs : Trying to change the state tree on a running instance."), __FUNCTION__);
+			return;
+		}
+	}
 
 	StateTreeRef = MoveTemp(InStateTreeReference);
 
@@ -452,12 +468,55 @@ void UStateTreeComponent::SetStateTreeReference(FStateTreeReference InStateTreeR
 		TValueOrError<void, FString> ValidResult = HasValidStateTreeReference();
 		if (ValidResult.HasError())
 		{
-			STATETREE_LOG(Warning, TEXT("%s: %s."), *ValidResult.GetError(), ANSI_TO_TCHAR(__FUNCTION__));
+			STATETREE_LOG(Warning, TEXT("%hs: %s."), __FUNCTION__, *ValidResult.GetError());
 			StateTreeRef = FStateTreeReference();
 		}
 	}
 }
 
+void UStateTreeComponent::SetLinkedStateTreeOverrides(FStateTreeReferenceOverrides Overrides)
+{
+	// Validate the schema
+	for (const FStateTreeReferenceOverrideItem& Item : Overrides.GetOverrideItems())
+	{
+		if (const UStateTree* ItemStateTree = Item.GetStateTreeReference().GetStateTree())
+		{
+			if (ItemStateTree->GetSchema() == nullptr
+				|| ItemStateTree->GetSchema()->GetClass()->IsChildOf(UStateTreeComponentSchema::StaticClass()))
+			{
+				STATETREE_LOG(Warning, TEXT("%hs: Trying to set the linked overrides '%s' with a wrong schema. %s."),
+					__FUNCTION__,
+					*Item.GetStateTag().ToString(),
+					*ItemStateTree->GetFullName()
+					);
+				return;
+			}
+		}
+	}
+
+	LinkedStateTreeOverrides = MoveTemp(Overrides);
+}
+
+void UStateTreeComponent::AddLinkedStateTreeOverrides(const FGameplayTag StateTag, FStateTreeReference StateTreeReference)
+{
+	// Validate the schema
+	if (const UStateTree* ItemStateTree = StateTreeReference.GetStateTree())
+	{
+		if (ItemStateTree->GetSchema() == nullptr
+			|| ItemStateTree->GetSchema()->GetClass()->IsChildOf(UStateTreeComponentSchema::StaticClass()))
+		{
+			STATETREE_LOG(Warning, TEXT("%hs: Trying to set the linked overrides with the wrong schema. %s."), __FUNCTION__, *ItemStateTree->GetFullName());
+			return;
+		}
+	}
+	LinkedStateTreeOverrides.AddOverride(FStateTreeReferenceOverrideItem(StateTag, MoveTemp(StateTreeReference)));
+}
+
+void UStateTreeComponent::RemoveLinkedStateTreeOverrides(const FGameplayTag StateTag)
+{
+	LinkedStateTreeOverrides.RemoveOverride(StateTag);
+}
+
 void UStateTreeComponent::SetStartLogicAutomatically(const bool bInStartLogicAutomatically)
 {
 	bStartLogicAutomatically = bInStartLogicAutomatically;
@@ -472,13 +531,13 @@ void UStateTreeComponent::SendStateTreeEvent(const FGameplayTag Tag, const FCons
 {
 	if (!bIsRunning)
 	{
-		STATETREE_LOG(Warning, TEXT("%s: Trying to send event to a State Tree that is not started yet."), ANSI_TO_TCHAR(__FUNCTION__));
+		STATETREE_LOG(Warning, TEXT("%hs: Trying to send event to a State Tree that is not started yet."), __FUNCTION__);
 		return;
 	}
 
 	if (!StateTreeRef.IsValid())
 	{
-		STATETREE_LOG(Warning, TEXT("%s: Trying to send event to State Tree component with invalid asset."), ANSI_TO_TCHAR(__FUNCTION__));
+		STATETREE_LOG(Warning, TEXT("%hs: Trying to send event to State Tree component with invalid asset."), __FUNCTION__);
 		return;
 	}
 
diff --git a/Engine/Plugins/Runtime/GameplayStateTree/Source/GameplayStateTreeModule/Public/Components/StateTreeComponent.h b/Engine/Plugins/Runtime/GameplayStateTree/Source/GameplayStateTreeModule/Public/Components/StateTreeComponent.h
--- a/Engine/Plugins/Runtime/GameplayStateTree/Source/GameplayStateTreeModule/Public/Components/StateTreeComponent.h
+++ b/Engine/Plugins/Runtime/GameplayStateTree/Source/GameplayStateTreeModule/Public/Components/StateTreeComponent.h
@@ -73,13 +73,36 @@ public:
 	UE_API TSubclassOf<UStateTreeSchema> GetSchema() const override;
 	//~ END
 
-	/** Stops currently running StateTree and sets new state tree reference. */
+	/**
+	 * Sets a new state tree.
+	 * The state tree won't be set if the logic is running.
+	 */
 	UFUNCTION(BlueprintCallable, Category = "Gameplay|StateTree")
 	UE_API void SetStateTree(UStateTree* StateTree);
 
-	/** Stops currently running StateTree and sets new state tree reference. */
+	/**
+	 * Sets a new state tree reference.
+	 * The state tree reference won't be set if the logic is running.
+	 */
 	UFUNCTION(BlueprintCallable, Category = "Gameplay|StateTree")
 	UE_API void SetStateTreeReference(FStateTreeReference StateTreeReference);
+
+	/**
+	 * Set the linked state tree overrides.
+	 * The overrides won't be set if they do not use the StateTreeComponentSchema schema.
+	 */
+	UE_API void SetLinkedStateTreeOverrides(FStateTreeReferenceOverrides Overrides);
+
+	/**
+	 * Add a linked state tree override.
+	 * The override won't be set if it doesn't use the StateTreeComponentSchema schema.
+	 */
+	UFUNCTION(BlueprintCallable, Category = "Gameplay|StateTree")
+	UE_API void AddLinkedStateTreeOverrides(const FGameplayTag StateTag, FStateTreeReference StateTreeReference);
+
+	/** Remove a linked state tree override. */
+	UFUNCTION(BlueprintCallable, Category = "Gameplay|StateTree")
+	UE_API void RemoveLinkedStateTreeOverrides(const FGameplayTag StateTag);
 	
 	/**
 	 * Sets whether the State Tree is started automatically on being play.
diff --git a/Engine/Plugins/Runtime/StateTree/Source/StateTreeModule/Public/StateTreeReference.h b/Engine/Plugins/Runtime/StateTree/Source/StateTreeModule/Public/StateTreeReference.h
--- a/Engine/Plugins/Runtime/StateTree/Source/StateTreeModule/Public/StateTreeReference.h
+++ b/Engine/Plugins/Runtime/StateTree/Source/StateTreeModule/Public/StateTreeReference.h
@@ -122,6 +122,11 @@ struct FStateTreeReferenceOverrideItem
 		, StateTreeReference(InStateTreeReference)
 	{
 	}
+	FStateTreeReferenceOverrideItem(const FGameplayTag InStateTag, FStateTreeReference&& InStateTreeReference)
+		: StateTag(InStateTag)
+		, StateTreeReference(MoveTemp(InStateTreeReference))
+	{
+	}
 	
 	FGameplayTag GetStateTag() const
 	{
@@ -183,6 +188,24 @@ struct FStateTreeReferenceOverrides
 		}
 	}
 
+	/** Adds or replaces override for a selected tag. */
+	void AddOverride(FStateTreeReferenceOverrideItem OverrideItem)
+	{
+		FStateTreeReferenceOverrideItem* FoundOverride = OverrideItems.FindByPredicate([&OverrideItem](const FStateTreeReferenceOverrideItem& Override)
+			{
+				return Override.GetStateTag() == OverrideItem.GetStateTag();
+			});
+
+		if (FoundOverride)
+		{
+			*FoundOverride = MoveTemp(OverrideItem);
+		}
+		else
+		{
+			OverrideItems.Emplace(MoveTemp(OverrideItem));
+		}
+	}
+
 	/** Returns true if removing an override succeeded. */
 	bool RemoveOverride(const FGameplayTag StateTag)
 	{
