[State Tree] async RunEnvQuery Task

diff --git a/Engine/Plugins/Runtime/GameplayStateTree/Source/GameplayStateTreeModule/Private/Tasks/StateTreeRunEnvQueryTask.cpp b/Engine/Plugins/Runtime/GameplayStateTree/Source/GameplayStateTreeModule/Private/Tasks/StateTreeRunEnvQueryTask.cpp
--- a/Engine/Plugins/Runtime/GameplayStateTree/Source/GameplayStateTreeModule/Private/Tasks/StateTreeRunEnvQueryTask.cpp
+++ b/Engine/Plugins/Runtime/GameplayStateTree/Source/GameplayStateTreeModule/Private/Tasks/StateTreeRunEnvQueryTask.cpp
@@ -4,10 +4,18 @@
 
 #include "EnvironmentQuery/EnvQuery.h"
 #include "EnvironmentQuery/EnvQueryManager.h"
+#include "StateTreeAsyncExecutionContext.h"
 #include "StateTreeExecutionContext.h"
 
 #define LOCTEXT_NAMESPACE "GameplayStateTree"
 
+FStateTreeRunEnvQueryTask::FStateTreeRunEnvQueryTask()
+{
+	bShouldCallTick = false;
+	bShouldCopyBoundPropertiesOnTick = false;
+	bShouldCopyBoundPropertiesOnExitState = false;
+}
+
 EStateTreeRunStatus FStateTreeRunEnvQueryTask::EnterState(FStateTreeExecutionContext& Context, const FStateTreeTransitionResult& Transition) const
 {
 	FInstanceDataType& InstanceData = Context.GetInstanceData(*this);
@@ -24,49 +32,42 @@ EStateTreeRunStatus FStateTreeRunEnvQueryTask::EnterState(FStateTreeExecutionCon
 	}
 
 	InstanceData.RequestId = Request.Execute(InstanceData.RunMode,
-		FQueryFinishedSignature::CreateLambda([InstanceDataRef = Context.GetInstanceDataStructRef(*this)](TSharedPtr<FEnvQueryResult> QueryResult) mutable
+		FQueryFinishedSignature::CreateLambda([WeakContext = Context.MakeWeakExecutionContext()](TSharedPtr<FEnvQueryResult> QueryResult) mutable
+		{
+			const FStateTreeStrongExecutionContext StrongContext = WeakContext.MakeStrongExecutionContext();
+			if (FInstanceDataType* InstanceDataPtr = StrongContext.GetInstanceDataPtr<FInstanceDataType>())
 			{
-				if (FInstanceDataType* InstanceData = InstanceDataRef.GetPtr())
+				InstanceDataPtr->RequestId = INDEX_NONE;
+
+				bool bSuccess = false;
+				if (QueryResult && QueryResult->IsSuccessful())
 				{
-					InstanceData->QueryResult = QueryResult;
-					InstanceData->RequestId = INDEX_NONE;
+					auto [VectorPtr, ActorPtr, ArrayOfVector, ArrayOfActor] = InstanceDataPtr->Result.GetPtrTupleFromStrongExecutionContext<FVector,AActor*, TArray<FVector>, TArray<AActor*>>(StrongContext);
+					if (VectorPtr)
+					{
+						*VectorPtr = QueryResult->GetItemAsLocation(0);
+					}
+					else if (ActorPtr)
+					{
+						*ActorPtr = QueryResult->GetItemAsActor(0);
+					}
+					else if (ArrayOfVector)
+					{
+						QueryResult->GetAllAsLocations(*ArrayOfVector);
+					}
+					else if (ArrayOfActor)
+					{
+						QueryResult->GetAllAsActors(*ArrayOfActor);
+					}
+
+					bSuccess = true;
 				}
-			}));
-	return InstanceData.RequestId != INDEX_NONE ? EStateTreeRunStatus::Running : EStateTreeRunStatus::Failed;
-}
 
-EStateTreeRunStatus FStateTreeRunEnvQueryTask::Tick(FStateTreeExecutionContext& Context, const float DeltaTime) const
-{
-	FInstanceDataType& InstanceData = Context.GetInstanceData(*this);
-	if (InstanceData.QueryResult)
-	{
-		if (InstanceData.QueryResult->IsSuccessful())
-		{
-			auto [VectorPtr, ActorPtr, ArrayOfVector, ArrayOfActor] = InstanceData.Result.GetMutablePtrTuple<FVector, AActor*, TArray<FVector>, TArray<AActor*>>(Context);
-			if (VectorPtr)
-			{
-				*VectorPtr = InstanceData.QueryResult->GetItemAsLocation(0);
+				StrongContext.FinishTask(bSuccess ? EStateTreeFinishTaskType::Succeeded : EStateTreeFinishTaskType::Failed);
 			}
-			else if (ActorPtr)
-			{
-				*ActorPtr = InstanceData.QueryResult->GetItemAsActor(0);
-			}
-			else if (ArrayOfVector)
-			{
-				InstanceData.QueryResult->GetAllAsLocations(*ArrayOfVector);
-			}
-			else if (ArrayOfActor)
-			{
-				InstanceData.QueryResult->GetAllAsActors(*ArrayOfActor);
-			}
-			return EStateTreeRunStatus::Succeeded;
-		}
-		else
-		{
-			return EStateTreeRunStatus::Failed;
-		}
-	}
-	return EStateTreeRunStatus::Running;
+		}));
+
+	return InstanceData.RequestId != INDEX_NONE ? EStateTreeRunStatus::Running : EStateTreeRunStatus::Failed;
 }
 
 void FStateTreeRunEnvQueryTask::ExitState(FStateTreeExecutionContext& Context, const FStateTreeTransitionResult& Transition) const
@@ -80,7 +81,6 @@ void FStateTreeRunEnvQueryTask::ExitState(FStateTreeExecutionContext& Context, c
 		}
 		InstanceData.RequestId = INDEX_NONE;
 	}
-	InstanceData.QueryResult.Reset();
 }
 
 #if WITH_EDITOR
diff --git a/Engine/Plugins/Runtime/GameplayStateTree/Source/GameplayStateTreeModule/Public/Tasks/StateTreeRunEnvQueryTask.h b/Engine/Plugins/Runtime/GameplayStateTree/Source/GameplayStateTreeModule/Public/Tasks/StateTreeRunEnvQueryTask.h
--- a/Engine/Plugins/Runtime/GameplayStateTree/Source/GameplayStateTreeModule/Public/Tasks/StateTreeRunEnvQueryTask.h
+++ b/Engine/Plugins/Runtime/GameplayStateTree/Source/GameplayStateTreeModule/Public/Tasks/StateTreeRunEnvQueryTask.h
@@ -3,10 +3,8 @@
 #pragma once
 
 #include "StateTreeTaskBase.h"
-
-#include "EnvironmentQuery/EnvQueryTypes.h"
-#include "NativeGameplayTags.h"
 #include "StateTreePropertyRef.h"
+#include "EnvironmentQuery/EnvQueryTypes.h"
 
 #include "StateTreeRunEnvQueryTask.generated.h"
 
@@ -35,8 +33,6 @@ struct FStateTreeRunEnvQueryInstanceData
 	UPROPERTY(EditAnywhere, Category = Parameter)
 	TEnumAsByte<EEnvQueryRunMode::Type> RunMode = EEnvQueryRunMode::SingleResult;
 
-	TSharedPtr<FEnvQueryResult> QueryResult = nullptr;
-
 	int32 RequestId = INDEX_NONE;
 };
 
@@ -54,10 +50,11 @@ struct FStateTreeRunEnvQueryTask : public FStateTreeTaskCommonBase
 
 	using FInstanceDataType = FStateTreeRunEnvQueryInstanceData;
 
+	FStateTreeRunEnvQueryTask();
+
 	virtual const UStruct* GetInstanceDataType() const override { return FInstanceDataType::StaticStruct(); }
 
 	virtual EStateTreeRunStatus EnterState(FStateTreeExecutionContext& Context, const FStateTreeTransitionResult& Transition) const override;
-	virtual EStateTreeRunStatus Tick(FStateTreeExecutionContext& Context, const float DeltaTime) const override;
 	virtual void ExitState(FStateTreeExecutionContext& Context, const FStateTreeTransitionResult& Transition) const override;
 
 #if WITH_EDITOR
diff --git a/Engine/Plugins/Runtime/StateTree/Source/StateTreeModule/Public/StateTreeTypes.h b/Engine/Plugins/Runtime/StateTree/Source/StateTreeModule/Public/StateTreeTypes.h
--- a/Engine/Plugins/Runtime/StateTree/Source/StateTreeModule/Public/StateTreeTypes.h
+++ b/Engine/Plugins/Runtime/StateTree/Source/StateTreeModule/Public/StateTreeTypes.h
@@ -869,11 +869,11 @@ struct FCompactStateTreeState
 	uint8 bHasTickTasksOnlyOnEvents : 1 = false;
 	
 	
-	/** True if any global tasks request a tick every frame. */
+	/** True if any state tasks request a tick every frame. */
 	uint8 bCachedRequestTick : 1 = false;
 
 	/**
-	 * True if any global tasks request a tick every frame but only if there are events.
+	 * True if any state tasks request a tick every frame but only if there are events.
 	 * No effect if bCachedRequestTick is true.
 	 */
 	uint8 bCachedRequestTickOnlyOnEvents : 1 = false;
