Setting UStateTreeAIComponent 'Blueprintable'

diff --git a/Engine/Plugins/Runtime/GameplayStateTree/Source/GameplayStateTreeModule/Public/Components/StateTreeAIComponent.h b/Engine/Plugins/Runtime/GameplayStateTree/Source/GameplayStateTreeModule/Public/Components/StateTreeAIComponent.h
--- a/Engine/Plugins/Runtime/GameplayStateTree/Source/GameplayStateTreeModule/Public/Components/StateTreeAIComponent.h
+++ b/Engine/Plugins/Runtime/GameplayStateTree/Source/GameplayStateTreeModule/Public/Components/StateTreeAIComponent.h
@@ -12,7 +12,7 @@
 * State tree component designed to be run on an AIController.
 * It uses the StateTreeAIComponentSchema that guarantees access to the AIController.
 */
-UCLASS(MinimalAPI, ClassGroup = AI, meta = (BlueprintSpawnableComponent))
+UCLASS(MinimalAPI, ClassGroup = AI, Blueprintable, meta = (BlueprintSpawnableComponent))
 class UStateTreeAIComponent : public UStateTreeComponent
 {
 	GENERATED_BODY()
