[State Tree] mixed up property names in MoveToTask

diff --git a/Engine/Plugins/Runtime/GameplayStateTree/Source/GameplayStateTreeModule/Private/Tasks/StateTreeMoveToTask.cpp b/Engine/Plugins/Runtime/GameplayStateTree/Source/GameplayStateTreeModule/Private/Tasks/StateTreeMoveToTask.cpp
--- a/Engine/Plugins/Runtime/GameplayStateTree/Source/GameplayStateTreeModule/Private/Tasks/StateTreeMoveToTask.cpp
+++ b/Engine/Plugins/Runtime/GameplayStateTree/Source/GameplayStateTreeModule/Private/Tasks/StateTreeMoveToTask.cpp
@@ -161,11 +161,11 @@ EDataValidationResult FStateTreeMoveToTask::Compile(UE::StateTree::ICompileNodeC
 	const FInstanceDataType& InstanceData = Context.GetInstanceDataView().Get<FInstanceDataType>();
 
 	// We only tick the task if we might track the destination vector
-	if (Context.HasBindingForProperty(FName(TEXT("TargetActor"))))
+	if (Context.HasBindingForProperty(FName(TEXT("Destination"))))
 	{
 		if (InstanceData.bTrackMovingGoal || Context.HasBindingForProperty(FName(TEXT("bTrackMovingGoal"))))
 		{
-			if (!InstanceData.TargetActor && !Context.HasBindingForProperty(FName(TEXT("Destination"))))
+			if (!InstanceData.TargetActor && !Context.HasBindingForProperty(FName(TEXT("TargetActor"))))
 			{
 				bShouldCallTick = true;
 				bShouldCopyBoundPropertiesOnTick = true;
