[StateTee] Added Blueprint helper to execute a StateTree on an actor (using existing StateTree component or adding it if necessary)
- component will now log errors during failed activation
- updated logging in UStateTreeComponentSchema to report in VLOG and engine log

diff --git a/Engine/Plugins/Runtime/GameplayStateTree/Source/GameplayStateTreeModule/Private/Components/StateTreeComponent.cpp b/Engine/Plugins/Runtime/GameplayStateTree/Source/GameplayStateTreeModule/Private/Components/StateTreeComponent.cpp
--- a/Engine/Plugins/Runtime/GameplayStateTree/Source/GameplayStateTreeModule/Private/Components/StateTreeComponent.cpp
+++ b/Engine/Plugins/Runtime/GameplayStateTree/Source/GameplayStateTreeModule/Private/Components/StateTreeComponent.cpp
@@ -83,7 +83,7 @@ bool UStateTreeComponent::SetContextRequirements(FStateTreeExecutionContext& Con
 {
 	Context.SetLinkedStateTreeOverrides(LinkedStateTreeOverrides);
 	Context.SetCollectExternalDataCallback(FOnCollectStateTreeExternalData::CreateUObject(this, &UStateTreeComponent::CollectExternalData));
-	return UStateTreeComponentSchema::SetContextRequirements(*this, Context);
+	return UStateTreeComponentSchema::SetContextRequirements(*this, Context, bLogErrors);
 }
 
 void UStateTreeComponent::BeginPlay()
@@ -165,7 +165,7 @@ void UStateTreeComponent::StartTree()
 	}
 
 	FStateTreeExecutionContext Context(*GetOwner(), *StateTreeRef.GetStateTree(), InstanceData);
-	if (SetContextRequirements(Context))
+	if (SetContextRequirements(Context, /*bLogErrors*/true))
 	{
 		const EStateTreeRunStatus PreviousRunStatus = Context.GetStateTreeRunStatus();
 
@@ -261,7 +261,7 @@ void UStateTreeComponent::ScheduleTickFrame(const FStateTreeScheduledTick& NextT
 	{
 		if (!UE::GameplayStateTree::Private::bScheduledTickAllowed)
 		{
-			// Make sure the component tick is enabled. It ticks every frames.
+			// Make sure the component tick is enabled. It ticks every frame.
 			if (!IsComponentTickEnabled())
 			{
 				SetComponentTickEnabled(true);
@@ -290,7 +290,7 @@ void UStateTreeComponent::ScheduleTickFrame(const FStateTreeScheduledTick& NextT
 			else
 			{
 				// We need to force a small dt to tell the TickTaskManager we might not want to be tick every frame.
-				const float FORCE_TICK_INTERVAL_DT = UE_KINDA_SMALL_NUMBER;
+				constexpr float FORCE_TICK_INTERVAL_DT = UE_KINDA_SMALL_NUMBER;
 				const float NextTickDeltaTime = !NextTick.ShouldTickOnceNextFrame() ? NextTick.GetTickRate() : FORCE_TICK_INTERVAL_DT;
 				if (!FMath::IsNearlyEqual(GetComponentTickInterval(), NextTickDeltaTime))
 				{
@@ -342,8 +342,7 @@ AActor* UStateTreeComponent::GetGameplayTaskOwner(const UGameplayTask* Task) con
 		return GetAIOwner();
 	}
 
-	const UAITask* AITask = Cast<const UAITask>(Task);
-	if (AITask)
+	if (const UAITask* AITask = Cast<const UAITask>(Task))
 	{
 		return AITask->GetAIController();
 	}
@@ -359,8 +358,7 @@ AActor* UStateTreeComponent::GetGameplayTaskAvatar(const UGameplayTask* Task) co
 		return GetAIOwner() ? GetAIOwner()->GetPawn() : nullptr;
 	}
 
-	const UAITask* AITask = Cast<const UAITask>(Task);
-	if (AITask)
+	if (const UAITask* AITask = Cast<const UAITask>(Task))
 	{
 		return AITask->GetAIController() ? AITask->GetAIController()->GetPawn() : nullptr;
 	}
@@ -380,7 +378,7 @@ void UStateTreeComponent::OnGameplayTaskInitialized(UGameplayTask& Task)
 	if (AITask && (AITask->GetAIController() == nullptr))
 	{
 		// this means that the task has either been created without specifying 
-		// UAITAsk::OwnerController's value (like via BP's Construct Object node)
+		// UAITask::OwnerController's value (like via BP's Construct Object node)
 		// or it has been created in C++ with inappropriate function
 		UE_LOG(LogStateTree, Error, TEXT("Missing AIController in AITask %s"), *AITask->GetName());
 	}
diff --git a/Engine/Plugins/Runtime/GameplayStateTree/Source/GameplayStateTreeModule/Private/Components/StateTreeComponentSchema.cpp b/Engine/Plugins/Runtime/GameplayStateTree/Source/GameplayStateTreeModule/Private/Components/StateTreeComponentSchema.cpp
--- a/Engine/Plugins/Runtime/GameplayStateTree/Source/GameplayStateTreeModule/Private/Components/StateTreeComponentSchema.cpp
+++ b/Engine/Plugins/Runtime/GameplayStateTree/Source/GameplayStateTreeModule/Private/Components/StateTreeComponentSchema.cpp
@@ -143,12 +143,12 @@ void UStateTreeComponentSchema::SetContextData(FContextDataSetter& ContextDataSe
 		}
 		if (ContextActor == nullptr && bLogErrors)
 		{
-			UE_VLOG(ContextDataSetter.GetComponent()->GetOwner(), LogStateTree, Error, TEXT("%s: Could not find context actor of type %s. StateTree will not update."), ANSI_TO_TCHAR(__FUNCTION__), *GetNameSafe(Schema->GetContextActorClass()));
+			UE_VLOG_UELOG(ContextDataSetter.GetComponent()->GetOwner(), LogStateTree, Error, TEXT("%hs: Could not find context actor of type %s. StateTree will not update."), __FUNCTION__, *GetNameSafe(Schema->GetContextActorClass()));
 		}
 	}
 	else if (bLogErrors)
 	{
-		UE_VLOG(ContextDataSetter.GetComponent()->GetOwner(), LogStateTree, Error, TEXT("%s: Expected StateTree asset to contain StateTreeComponentSchema. StateTree will not update."), ANSI_TO_TCHAR(__FUNCTION__));
+		UE_VLOG_UELOG(ContextDataSetter.GetComponent()->GetOwner(), LogStateTree, Error, TEXT("%hs: Expected StateTree asset to contain StateTreeComponentSchema. StateTree will not update."), __FUNCTION__);
 	}
 
 	const FName ActorName = UE::GameplayStateTree::Private::Name_Actor;
@@ -168,7 +168,7 @@ bool UStateTreeComponentSchema::SetContextRequirements(UBrainComponent& BrainCom
 	bool bResult = Context.AreContextDataViewsValid();
 	if (!bResult && bLogErrors)
 	{
-		UE_VLOG(BrainComponent.GetOwner(), LogStateTree, Error, TEXT("%s: Missing external data requirements. StateTree will not update."), ANSI_TO_TCHAR(__FUNCTION__));
+		UE_VLOG_UELOG(BrainComponent.GetOwner(), LogStateTree, Error, TEXT("%hs: Missing external data requirements. StateTree will not update."), __FUNCTION__);
 	}
 
 	return bResult;
@@ -181,7 +181,7 @@ bool UStateTreeComponentSchema::CollectExternalData(const FStateTreeExecutionCon
 	const UWorld* World = Context.GetWorld();
 	if (World == nullptr)
 	{
-		UE_VLOG(Context.GetOwner(), LogStateTree, Error, TEXT("%hs: External data can't find the world ('%s' using StateTree '%s')."),
+		UE_VLOG_UELOG(Context.GetOwner(), LogStateTree, Error, TEXT("%hs: External data can't find the world ('%s' using StateTree '%s')."),
 			__FUNCTION__, *GetNameSafe(Context.GetOwner()), *GetFullNameSafe(Context.GetStateTree()));
 		return false;
 	}
@@ -189,7 +189,7 @@ bool UStateTreeComponentSchema::CollectExternalData(const FStateTreeExecutionCon
 	AActor* Owner = Cast<AActor>(Context.GetOwner());
 	if (!Owner)
 	{
-		UE_VLOG(Context.GetOwner(), LogStateTree, Error, TEXT("%hs: The external data owner is invalid ('%s' using StateTree '%s')."),
+		UE_VLOG_UELOG(Context.GetOwner(), LogStateTree, Error, TEXT("%hs: The external data owner is invalid ('%s' using StateTree '%s')."),
 			__FUNCTION__, *GetNameSafe(Context.GetOwner()), *GetFullNameSafe(Context.GetStateTree()));
 		return false;
 	}
diff --git a/Engine/Plugins/Runtime/GameplayStateTree/Source/GameplayStateTreeModule/Private/GameplayStateTreeBlueprintFunctionLibrary.cpp b/Engine/Plugins/Runtime/GameplayStateTree/Source/GameplayStateTreeModule/Private/GameplayStateTreeBlueprintFunctionLibrary.cpp
new file mode 100644
--- /dev/null
+++ b/Engine/Plugins/Runtime/GameplayStateTree/Source/GameplayStateTreeModule/Private/GameplayStateTreeBlueprintFunctionLibrary.cpp
@@ -0,0 +1,74 @@
+// Copyright Epic Games, Inc. All Rights Reserved.
+
+#include "GameplayStateTreeBlueprintFunctionLibrary.h"
+#include "AIController.h"
+#include "Blueprint/AIBlueprintHelperLibrary.h"
+#include "Components/StateTreeComponent.h"
+#include "StateTree.h"
+#include "StateTreeTypes.h"
+#include "VisualLogger/VisualLogger.h"
+
+#include UE_INLINE_GENERATED_CPP_BY_NAME(GameplayStateTreeBlueprintFunctionLibrary)
+
+bool UGameplayStateTreeBlueprintFunctionLibrary::RunStateTree(AActor* Actor, UStateTree* StateTreeAsset)
+{
+	if (Actor == nullptr)
+	{
+		FFrame::KismetExecutionMessage(*FString::Printf(TEXT("Unable to run an instance of statetree '%s': null actor")
+			, *GetNameSafe(StateTreeAsset))
+			, ELogVerbosity::Warning);
+		return false;
+	}
+
+	if (StateTreeAsset == nullptr)
+	{
+		FFrame::KismetExecutionMessage(*FString::Printf(TEXT("Unable to run an instance of statetree on actor '%s': null state tree asset")
+			, *Actor->GetName())
+			, ELogVerbosity::Warning);
+		return false;
+	}
+
+	AAIController* AIController = Cast<AAIController>(Actor);
+	UStateTreeComponent* StateTreeComponent = nullptr;
+	if (AIController != nullptr)
+	{
+		StateTreeComponent = Cast<UStateTreeComponent>(AIController->GetBrainComponent());
+	}
+	else
+	{
+		StateTreeComponent = Actor->GetComponentByClass<UStateTreeComponent>();
+	}
+
+	if (StateTreeComponent == nullptr)
+	{
+		StateTreeComponent = NewObject<UStateTreeComponent>(Actor);
+
+		// Make sure to not attempt to start the logic until the component is set up
+		StateTreeComponent->SetStartLogicAutomatically(false);
+		StateTreeComponent->RegisterComponent();
+		REDIRECT_OBJECT_TO_VLOG(StateTreeComponent, Actor);
+	}
+
+	// Make sure BrainComponent points at the newly created StateTree component
+	if (AIController != nullptr)
+	{
+		AIController->BrainComponent = StateTreeComponent;
+	}
+
+	check(StateTreeComponent);
+	if (StateTreeComponent->IsRunning())
+	{
+		StateTreeComponent->StopLogic(TEXT("Starting logic with new asset"));
+	}
+
+	StateTreeComponent->SetStateTree(StateTreeAsset);
+	StateTreeComponent->SetStartLogicAutomatically(true);
+
+	// If the component was already simulating we start the logic, otherwise we set it to auto-start on BeginPlay
+	if (StateTreeComponent->HasBegunPlay())
+	{
+		StateTreeComponent->StartLogic();
+	}
+
+	return true;
+}
\ No newline at end of file
diff --git a/Engine/Plugins/Runtime/GameplayStateTree/Source/GameplayStateTreeModule/Public/Components/StateTreeComponent.h b/Engine/Plugins/Runtime/GameplayStateTree/Source/GameplayStateTreeModule/Public/Components/StateTreeComponent.h
--- a/Engine/Plugins/Runtime/GameplayStateTree/Source/GameplayStateTreeModule/Public/Components/StateTreeComponent.h
+++ b/Engine/Plugins/Runtime/GameplayStateTree/Source/GameplayStateTreeModule/Public/Components/StateTreeComponent.h
@@ -70,7 +70,7 @@ public:
 	//~ END IGameplayTaskOwnerInterface
 
 	//~ BEGIN IStateTreeSchemaProvider
-	UE_API TSubclassOf<UStateTreeSchema> GetSchema() const override;
+	UE_API virtual TSubclassOf<UStateTreeSchema> GetSchema() const override;
 	//~ END
 
 	/**
@@ -105,7 +105,7 @@ public:
 	UE_API void RemoveLinkedStateTreeOverrides(const FGameplayTag StateTag);
 	
 	/**
-	 * Sets whether the State Tree is started automatically on being play.
+	 * Sets whether the State Tree is started automatically on begin play.
 	 * This function sets the bStartLogicAutomatically property, and should be used mostly from constructions scripts.
 	 * If you wish to start the logic manually, call StartLogic(). 
 	 */
@@ -144,7 +144,7 @@ protected:
 	UE_API virtual void PostLoad() override;
 #endif
 	/**
-	 * Called during initialize, will validate the state tree reference and create a context from the state tree to check it's validity
+	 * Called during initialize, will validate the state tree reference and create a context from the state tree to check its validity
 	 * Override this function for custom state tree validation.
 	 * Note: Override without calling super if the state tree reference is dynamically set after initialization
 	 */
@@ -182,7 +182,7 @@ protected:
 	UPROPERTY(Transient)
 	FStateTreeInstanceData InstanceData;
 
-	/** If true, the StateTree logic is started on begin play. Otherwise StartLogic() needs to be called. */
+	/** If true, the StateTree logic is started on begin play. Otherwise, StartLogic() needs to be called. */
 	UPROPERTY(EditAnywhere, Category = AI)
 	bool bStartLogicAutomatically = true;
 
diff --git a/Engine/Plugins/Runtime/GameplayStateTree/Source/GameplayStateTreeModule/Public/GameplayStateTreeBlueprintFunctionLibrary.h b/Engine/Plugins/Runtime/GameplayStateTree/Source/GameplayStateTreeModule/Public/GameplayStateTreeBlueprintFunctionLibrary.h
new file mode 100644
--- /dev/null
+++ b/Engine/Plugins/Runtime/GameplayStateTree/Source/GameplayStateTreeModule/Public/GameplayStateTreeBlueprintFunctionLibrary.h
@@ -0,0 +1,30 @@
+// Copyright Epic Games, Inc. All Rights Reserved.
+
+#pragma once
+
+#include "Kismet/BlueprintFunctionLibrary.h"
+#include "GameplayStateTreeBlueprintFunctionLibrary.generated.h"
+
+#define UE_API GAMEPLAYSTATETREEMODULE_API
+
+class UStateTree;
+class AActor;
+
+UCLASS(MinimalAPI, meta = (ScriptName = "GameplayStateTreeLibrary"))
+class UGameplayStateTreeBlueprintFunctionLibrary : public UBlueprintFunctionLibrary
+{
+	GENERATED_BODY()
+public:
+	
+	/**
+	 * Starts executing an instance of the provided state tree asset on the specified actor.
+	 * This will use the existing required component if present or add it if missing.
+	 * @param Actor The actor on which the StateTree must be executed.
+	 * @param StateTreeAsset The state tree asset to execute
+	 * @return Whether a state tree instance was successfully started
+	 */
+	UFUNCTION(BlueprintCallable, Category = "AI", meta=(ReturnDisplayName="bSuccess"))
+	static UE_API bool RunStateTree(AActor* Actor, UStateTree* StateTreeAsset);
+};
+
+#undef UE_API
diff --git a/Engine/Source/Runtime/AIModule/Private/AIController.cpp b/Engine/Source/Runtime/AIModule/Private/AIController.cpp
--- a/Engine/Source/Runtime/AIModule/Private/AIController.cpp
+++ b/Engine/Source/Runtime/AIModule/Private/AIController.cpp
@@ -1008,7 +1008,7 @@ bool AAIController::RunBehaviorTree(UBehaviorTree* BTAsset)
 		{
 			UE_VLOG(this, LogBehaviorTree, Log, TEXT("RunBehaviorTree: spawning BehaviorTreeComponent.."));
 
-			BTComp = NewObject<UBehaviorTreeComponent>(this, TEXT("BTComponent"));
+			BTComp = NewObject<UBehaviorTreeComponent>(this);
 			BTComp->RegisterComponent();
 			REDIRECT_OBJECT_TO_VLOG(BTComp, this);
 		}
