[StateTree] fixed check firing in UStateTreeComponent::SetStateTree by testing the current status only if the associated asset is valid

diff --git a/Engine/Plugins/Runtime/GameplayStateTree/Source/GameplayStateTreeModule/Private/Components/StateTreeComponent.cpp b/Engine/Plugins/Runtime/GameplayStateTree/Source/GameplayStateTreeModule/Private/Components/StateTreeComponent.cpp
--- a/Engine/Plugins/Runtime/GameplayStateTree/Source/GameplayStateTreeModule/Private/Components/StateTreeComponent.cpp
+++ b/Engine/Plugins/Runtime/GameplayStateTree/Source/GameplayStateTreeModule/Private/Components/StateTreeComponent.cpp
@@ -424,9 +424,10 @@ void UStateTreeComponent::ValidateStateTreeReference()
 
 void UStateTreeComponent::SetStateTree(UStateTree* InStateTree)
 {
+	if (const UStateTree* ActiveStateTree = StateTreeRef.GetStateTree())
 	{
 		// Can't change the StateTreeRef on a running tree. It might change the instance data while running tasks.
-		FStateTreeReadOnlyExecutionContext Context(GetOwner(), StateTreeRef.GetStateTree(), InstanceData);
+		FStateTreeReadOnlyExecutionContext Context(GetOwner(), ActiveStateTree, InstanceData);
 		if (Context.GetStateTreeRunStatus() == EStateTreeRunStatus::Running)
 		{
 			STATETREE_LOG(Warning, TEXT("%hs : Trying to change the state tree on a running instance."), __FUNCTION__);
