[State Tree] Exported virtual functions for MoveTo and RunEnvQuery tasks

diff --git a/Engine/Plugins/Runtime/GameplayStateTree/Source/GameplayStateTreeModule/Public/Tasks/StateTreeMoveToTask.h b/Engine/Plugins/Runtime/GameplayStateTree/Source/GameplayStateTreeModule/Public/Tasks/StateTreeMoveToTask.h
--- a/Engine/Plugins/Runtime/GameplayStateTree/Source/GameplayStateTreeModule/Public/Tasks/StateTreeMoveToTask.h
+++ b/Engine/Plugins/Runtime/GameplayStateTree/Source/GameplayStateTreeModule/Public/Tasks/StateTreeMoveToTask.h
@@ -9,6 +9,8 @@
 
 #include "StateTreeMoveToTask.generated.h"
 
+#define UE_API GAMEPLAYSTATETREEMODULE_API
+
 class AActor;
 class AAIController;
 class IGameplayTaskOwnerInterface;
@@ -91,18 +93,18 @@ struct FStateTreeMoveToTask : public FStateTreeAIActionTaskBase
 
 	virtual const UStruct* GetInstanceDataType() const override { return FInstanceDataType::StaticStruct(); }
 
-	GAMEPLAYSTATETREEMODULE_API FStateTreeMoveToTask();
-	GAMEPLAYSTATETREEMODULE_API virtual EStateTreeRunStatus EnterState(FStateTreeExecutionContext& Context, const FStateTreeTransitionResult& Transition) const override;
-	GAMEPLAYSTATETREEMODULE_API virtual EStateTreeRunStatus Tick(FStateTreeExecutionContext& Context, const float DeltaTime) const override;
-	GAMEPLAYSTATETREEMODULE_API virtual void ExitState(FStateTreeExecutionContext& Context, const FStateTreeTransitionResult& Transition) const override;
+	UE_API FStateTreeMoveToTask();
+	UE_API virtual EStateTreeRunStatus EnterState(FStateTreeExecutionContext& Context, const FStateTreeTransitionResult& Transition) const override;
+	UE_API virtual EStateTreeRunStatus Tick(FStateTreeExecutionContext& Context, const float DeltaTime) const override;
+	UE_API virtual void ExitState(FStateTreeExecutionContext& Context, const FStateTreeTransitionResult& Transition) const override;
 
-	GAMEPLAYSTATETREEMODULE_API virtual UAITask_MoveTo* PrepareMoveToTask(FStateTreeExecutionContext& Context, AAIController& Controller, UAITask_MoveTo* ExistingTask, FAIMoveRequest& MoveRequest) const;
-	GAMEPLAYSTATETREEMODULE_API virtual EStateTreeRunStatus PerformMoveTask(FStateTreeExecutionContext& Context, AAIController& Controller) const;
+	UE_API virtual UAITask_MoveTo* PrepareMoveToTask(FStateTreeExecutionContext& Context, AAIController& Controller, UAITask_MoveTo* ExistingTask, FAIMoveRequest& MoveRequest) const;
+	UE_API virtual EStateTreeRunStatus PerformMoveTask(FStateTreeExecutionContext& Context, AAIController& Controller) const;
 
 #if WITH_EDITOR
-	virtual EDataValidationResult Compile(UE::StateTree::ICompileNodeContext& Context) override;
+	UE_API virtual EDataValidationResult Compile(UE::StateTree::ICompileNodeContext& Context) override;
 
-	virtual FText GetDescription(const FGuid& ID, FStateTreeDataView InstanceDataView, const IStateTreeBindingLookup& BindingLookup, EStateTreeNodeFormatting Formatting = EStateTreeNodeFormatting::Text) const override;
+	UE_API virtual FText GetDescription(const FGuid& ID, FStateTreeDataView InstanceDataView, const IStateTreeBindingLookup& BindingLookup, EStateTreeNodeFormatting Formatting = EStateTreeNodeFormatting::Text) const override;
 	virtual FName GetIconName() const override
 	{
 		return FName("StateTreeEditorStyle|Node.Movement");
@@ -113,4 +115,6 @@ struct FStateTreeMoveToTask : public FStateTreeAIActionTaskBase
 	}
 #endif // WITH_EDITOR
 
-};
\ No newline at end of file
+};
+
+#undef UE_API
\ No newline at end of file
diff --git a/Engine/Plugins/Runtime/GameplayStateTree/Source/GameplayStateTreeModule/Public/Tasks/StateTreeRunEnvQueryTask.h b/Engine/Plugins/Runtime/GameplayStateTree/Source/GameplayStateTreeModule/Public/Tasks/StateTreeRunEnvQueryTask.h
--- a/Engine/Plugins/Runtime/GameplayStateTree/Source/GameplayStateTreeModule/Public/Tasks/StateTreeRunEnvQueryTask.h
+++ b/Engine/Plugins/Runtime/GameplayStateTree/Source/GameplayStateTreeModule/Public/Tasks/StateTreeRunEnvQueryTask.h
@@ -8,6 +8,8 @@
 
 #include "StateTreeRunEnvQueryTask.generated.h"
 
+#define UE_API GAMEPLAYSTATETREEMODULE_API
+
 USTRUCT()
 struct FStateTreeRunEnvQueryInstanceData
 {
@@ -50,16 +52,16 @@ struct FStateTreeRunEnvQueryTask : public FStateTreeTaskCommonBase
 
 	using FInstanceDataType = FStateTreeRunEnvQueryInstanceData;
 
-	FStateTreeRunEnvQueryTask();
+	UE_API FStateTreeRunEnvQueryTask();
 
 	virtual const UStruct* GetInstanceDataType() const override { return FInstanceDataType::StaticStruct(); }
 
-	virtual EStateTreeRunStatus EnterState(FStateTreeExecutionContext& Context, const FStateTreeTransitionResult& Transition) const override;
-	virtual void ExitState(FStateTreeExecutionContext& Context, const FStateTreeTransitionResult& Transition) const override;
+	UE_API virtual EStateTreeRunStatus EnterState(FStateTreeExecutionContext& Context, const FStateTreeTransitionResult& Transition) const override;
+	UE_API virtual void ExitState(FStateTreeExecutionContext& Context, const FStateTreeTransitionResult& Transition) const override;
 
 #if WITH_EDITOR
-	virtual void PostEditInstanceDataChangeChainProperty(const FPropertyChangedChainEvent& PropertyChangedEvent, FStateTreeDataView InstanceDataView) override;
-	virtual FText GetDescription(const FGuid& ID, FStateTreeDataView InstanceDataView, const IStateTreeBindingLookup& BindingLookup, EStateTreeNodeFormatting Formatting = EStateTreeNodeFormatting::Text) const override;
+	UE_API virtual void PostEditInstanceDataChangeChainProperty(const FPropertyChangedChainEvent& PropertyChangedEvent, FStateTreeDataView InstanceDataView) override;
+	UE_API virtual FText GetDescription(const FGuid& ID, FStateTreeDataView InstanceDataView, const IStateTreeBindingLookup& BindingLookup, EStateTreeNodeFormatting Formatting = EStateTreeNodeFormatting::Text) const override;
 	virtual FName GetIconName() const override
 	{
 		return FName("StateTreeEditorStyle|Node.Find");
@@ -70,3 +72,5 @@ struct FStateTreeRunEnvQueryTask : public FStateTreeTaskCommonBase
 	}
 #endif // WITH_EDITOR
 };
+
+#undef UE_API
\ No newline at end of file
