[StateTree] do not validate the StateTreeReference on initialization if the component is not setup to start the logic automatically. In this case it is valid to not have an asset set.

diff --git a/Engine/Plugins/Runtime/GameplayStateTree/Source/GameplayStateTreeModule/Private/Components/StateTreeComponent.cpp b/Engine/Plugins/Runtime/GameplayStateTree/Source/GameplayStateTreeModule/Private/Components/StateTreeComponent.cpp
--- a/Engine/Plugins/Runtime/GameplayStateTree/Source/GameplayStateTreeModule/Private/Components/StateTreeComponent.cpp
+++ b/Engine/Plugins/Runtime/GameplayStateTree/Source/GameplayStateTreeModule/Private/Components/StateTreeComponent.cpp
@@ -49,7 +49,11 @@ void UStateTreeComponent::InitializeComponent()
 {
 	// Skipping UBrainComponent
 	UActorComponent::InitializeComponent();
-	ValidateStateTreeReference();
+
+	if (bStartLogicAutomatically)
+	{
+		ValidateStateTreeReference();
+	}
 }
 
 #if WITH_EDITOR
