[StateTree] do not validate the StateTreeReference on initialization if the component is not set up to start the logic automatically.
In this case, it is valid not to have an asset set.

diff --git a/Engine/Plugins/Runtime/GameplayStateTree/Source/GameplayStateTreeModule/Private/Components/StateTreeComponent.cpp b/Engine/Plugins/Runtime/GameplayStateTree/Source/GameplayStateTreeModule/Private/Components/StateTreeComponent.cpp
--- a/Engine/Plugins/Runtime/GameplayStateTree/Source/GameplayStateTreeModule/Private/Components/StateTreeComponent.cpp
+++ b/Engine/Plugins/Runtime/GameplayStateTree/Source/GameplayStateTreeModule/Private/Components/StateTreeComponent.cpp
@@ -410,7 +410,7 @@ void UStateTreeComponent::ValidateStateTreeReference()
 	}
 }
 
- TValueOrError<void, FString> UStateTreeComponent::HasValidStateTreeReference() const
+TValueOrError<void, FString> UStateTreeComponent::HasValidStateTreeReference() const
 {
 	if (!StateTreeRef.IsValid())
 	{
@@ -443,12 +443,7 @@ void UStateTreeComponent::SetStateTree(UStateTree* InStateTree)
 
 	if (StateTreeRef.GetStateTree() != nullptr)
 	{
-		TValueOrError<void, FString> ValidResult = HasValidStateTreeReference();
-		if (ValidResult.HasError())
-		{
-			STATETREE_LOG(Warning, TEXT("%hs: %s."), __FUNCTION__, *ValidResult.GetError());
-			StateTreeRef = FStateTreeReference();
-		}
+		ValidateStateTreeReference();
 	}
 }
 
@@ -468,12 +463,7 @@ void UStateTreeComponent::SetStateTreeReference(FStateTreeReference InStateTreeR
 
 	if (StateTreeRef.GetStateTree() != nullptr)
 	{
-		TValueOrError<void, FString> ValidResult = HasValidStateTreeReference();
-		if (ValidResult.HasError())
-		{
-			STATETREE_LOG(Warning, TEXT("%hs: %s."), __FUNCTION__, *ValidResult.GetError());
-			StateTreeRef = FStateTreeReference();
-		}
+		ValidateStateTreeReference();
 	}
 }
 
