[StateTree] Fix schema test in StateTreeComponent.

diff --git a/Engine/Plugins/Runtime/GameplayStateTree/Source/GameplayStateTreeModule/Private/Components/StateTreeComponent.cpp b/Engine/Plugins/Runtime/GameplayStateTree/Source/GameplayStateTreeModule/Private/Components/StateTreeComponent.cpp
--- a/Engine/Plugins/Runtime/GameplayStateTree/Source/GameplayStateTreeModule/Private/Components/StateTreeComponent.cpp
+++ b/Engine/Plugins/Runtime/GameplayStateTree/Source/GameplayStateTreeModule/Private/Components/StateTreeComponent.cpp
@@ -472,12 +472,18 @@ void UStateTreeComponent::SetStateTreeReference(FStateTreeReference InStateTreeR
 void UStateTreeComponent::SetLinkedStateTreeOverrides(FStateTreeReferenceOverrides Overrides)
 {
 	// Validate the schema
+	TSubclassOf<UStateTreeComponentSchema> BaseSchema = UStateTreeComponentSchema::StaticClass();
+	if (StateTreeRef.GetStateTree() && StateTreeRef.GetStateTree()->GetSchema())
+	{
+		BaseSchema = StateTreeRef.GetStateTree()->GetSchema()->GetClass();
+	}
+
 	for (const FStateTreeReferenceOverrideItem& Item : Overrides.GetOverrideItems())
 	{
 		if (const UStateTree* ItemStateTree = Item.GetStateTreeReference().GetStateTree())
 		{
 			if (ItemStateTree->GetSchema() == nullptr
-				|| ItemStateTree->GetSchema()->GetClass()->IsChildOf(UStateTreeComponentSchema::StaticClass()))
+				|| !ItemStateTree->GetSchema()->GetClass()->IsChildOf(BaseSchema.Get()))
 			{
 				STATETREE_LOG(Warning, TEXT("%hs: Trying to set the linked overrides '%s' with a wrong schema. %s."),
 					__FUNCTION__,
@@ -497,8 +503,14 @@ void UStateTreeComponent::AddLinkedStateTreeOverrides(const FGameplayTag StateTa
 	// Validate the schema
 	if (const UStateTree* ItemStateTree = StateTreeReference.GetStateTree())
 	{
+		TSubclassOf<UStateTreeComponentSchema> BaseSchema = UStateTreeComponentSchema::StaticClass();
+		if (StateTreeRef.GetStateTree() && StateTreeRef.GetStateTree()->GetSchema())
+		{
+			BaseSchema = StateTreeRef.GetStateTree()->GetSchema()->GetClass();
+		}
+
 		if (ItemStateTree->GetSchema() == nullptr
-			|| ItemStateTree->GetSchema()->GetClass()->IsChildOf(UStateTreeComponentSchema::StaticClass()))
+			|| !ItemStateTree->GetSchema()->GetClass()->IsChildOf(BaseSchema.Get()))
 		{
 			STATETREE_LOG(Warning, TEXT("%hs: Trying to set the linked overrides with the wrong schema. %s."), __FUNCTION__, *ItemStateTree->GetFullName());
 			return;
