PR #13579: Fix StateTreeMoveToTask to persist bShouldCallTick for packaged builds

diff --git a/Engine/Plugins/Runtime/GameplayStateTree/Source/GameplayStateTreeModule/Private/Tasks/StateTreeMoveToTask.cpp b/Engine/Plugins/Runtime/GameplayStateTree/Source/GameplayStateTreeModule/Private/Tasks/StateTreeMoveToTask.cpp
--- a/Engine/Plugins/Runtime/GameplayStateTree/Source/GameplayStateTreeModule/Private/Tasks/StateTreeMoveToTask.cpp
+++ b/Engine/Plugins/Runtime/GameplayStateTree/Source/GameplayStateTreeModule/Private/Tasks/StateTreeMoveToTask.cpp
@@ -20,6 +20,14 @@ FStateTreeMoveToTask::FStateTreeMoveToTask()
 	bShouldCopyBoundPropertiesOnExitState = false;
 }
 
+bool FStateTreeMoveToTask::Link(FStateTreeLinker& Linker)
+{
+	bShouldCallTick = bSavedShouldCallTick;
+	bShouldCopyBoundPropertiesOnTick = bSavedShouldCallTick;
+
+	return Super::Link(Linker);
+}
+
 EStateTreeRunStatus FStateTreeMoveToTask::EnterState(FStateTreeExecutionContext& Context, const FStateTreeTransitionResult& Transition) const
 {
 	FInstanceDataType& InstanceData = Context.GetInstanceData(*this);
@@ -167,8 +175,8 @@ EDataValidationResult FStateTreeMoveToTask::Compile(UE::StateTree::ICompileNodeC
 		{
 			if (!InstanceData.TargetActor && !Context.HasBindingForProperty(FName(TEXT("TargetActor"))))
 			{
-				bShouldCallTick = true;
-				bShouldCopyBoundPropertiesOnTick = true;
+				// Flags on the base task are not saved. They will be set in Link().
+				bSavedShouldCallTick = true;
 			}
 		}
 	}
diff --git a/Engine/Plugins/Runtime/GameplayStateTree/Source/GameplayStateTreeModule/Public/Tasks/StateTreeMoveToTask.h b/Engine/Plugins/Runtime/GameplayStateTree/Source/GameplayStateTreeModule/Public/Tasks/StateTreeMoveToTask.h
--- a/Engine/Plugins/Runtime/GameplayStateTree/Source/GameplayStateTreeModule/Public/Tasks/StateTreeMoveToTask.h
+++ b/Engine/Plugins/Runtime/GameplayStateTree/Source/GameplayStateTreeModule/Public/Tasks/StateTreeMoveToTask.h
@@ -94,6 +94,7 @@ struct FStateTreeMoveToTask : public FStateTreeAIActionTaskBase
 	virtual const UStruct* GetInstanceDataType() const override { return FInstanceDataType::StaticStruct(); }
 
 	UE_API FStateTreeMoveToTask();
+	UE_API virtual bool Link(FStateTreeLinker& Linker) override;
 	UE_API virtual EStateTreeRunStatus EnterState(FStateTreeExecutionContext& Context, const FStateTreeTransitionResult& Transition) const override;
 	UE_API virtual EStateTreeRunStatus Tick(FStateTreeExecutionContext& Context, const float DeltaTime) const override;
 	UE_API virtual void ExitState(FStateTreeExecutionContext& Context, const FStateTreeTransitionResult& Transition) const override;
@@ -115,6 +116,8 @@ struct FStateTreeMoveToTask : public FStateTreeAIActionTaskBase
 	}
 #endif // WITH_EDITOR
 
+	UPROPERTY()
+	bool bSavedShouldCallTick = false;
 };
 
 #undef UE_API
\ No newline at end of file
