Fixes a crash when calling SetStateTreeReference on a runtime created StateTreeComponent

diff --git a/Engine/Plugins/Runtime/GameplayStateTree/Source/GameplayStateTreeModule/Private/Components/StateTreeComponent.cpp b/Engine/Plugins/Runtime/GameplayStateTree/Source/GameplayStateTreeModule/Private/Components/StateTreeComponent.cpp
--- a/Engine/Plugins/Runtime/GameplayStateTree/Source/GameplayStateTreeModule/Private/Components/StateTreeComponent.cpp
+++ b/Engine/Plugins/Runtime/GameplayStateTree/Source/GameplayStateTreeModule/Private/Components/StateTreeComponent.cpp
@@ -451,6 +451,7 @@ void UStateTreeComponent::SetStateTree(UStateTree* InStateTree)
 
 void UStateTreeComponent::SetStateTreeReference(FStateTreeReference InStateTreeReference)
 {
+	if (const UStateTree* const ActiveStateTree = StateTreeRef.GetStateTree())
 	{
 		// Can't change the StateTreeRef on a running tree. It might change the instance data while running tasks.
 		FStateTreeReadOnlyExecutionContext Context(GetOwner(), StateTreeRef.GetStateTree(), InstanceData);
