[StateTree] AITask_MoveTo requires high priority to be able to run inside other tasks (e.g., GameplayInteraction task). This is now consistent with the other uses of that task.

diff --git a/Engine/Plugins/Runtime/GameplayStateTree/Source/GameplayStateTreeModule/Private/Tasks/StateTreeMoveToTask.cpp b/Engine/Plugins/Runtime/GameplayStateTree/Source/GameplayStateTreeModule/Private/Tasks/StateTreeMoveToTask.cpp
--- a/Engine/Plugins/Runtime/GameplayStateTree/Source/GameplayStateTreeModule/Private/Tasks/StateTreeMoveToTask.cpp
+++ b/Engine/Plugins/Runtime/GameplayStateTree/Source/GameplayStateTreeModule/Private/Tasks/StateTreeMoveToTask.cpp
@@ -83,7 +83,7 @@ void FStateTreeMoveToTask::ExitState(FStateTreeExecutionContext& Context, const
 UAITask_MoveTo* FStateTreeMoveToTask::PrepareMoveToTask(FStateTreeExecutionContext& Context, AAIController& Controller, UAITask_MoveTo* ExistingTask, FAIMoveRequest& MoveRequest) const
 {
 	FInstanceDataType& InstanceData = Context.GetInstanceData(*this);
-	UAITask_MoveTo* MoveTask = ExistingTask ? ExistingTask : UAITask::NewAITask<UAITask_MoveTo>(Controller, *InstanceData.TaskOwner);
+	UAITask_MoveTo* MoveTask = ExistingTask ? ExistingTask : UAITask::NewAITask<UAITask_MoveTo>(Controller, *InstanceData.TaskOwner, EAITaskPriority::High);
 	if (MoveTask)
 	{
 		MoveTask->SetUp(&Controller, MoveRequest);
