[State Tree] correctly set running flag in re-entrance for State Tree Component

Set IsRunning true only if the tree execution has successfully started. Set IsRunning false only if the tree execution has successfully stopped or StateTreeReference has been found invalid.

diff --git a/Engine/Plugins/Runtime/GameplayStateTree/Source/GameplayStateTreeModule/Private/Components/StateTreeComponent.cpp b/Engine/Plugins/Runtime/GameplayStateTree/Source/GameplayStateTreeModule/Private/Components/StateTreeComponent.cpp
--- a/Engine/Plugins/Runtime/GameplayStateTree/Source/GameplayStateTreeModule/Private/Components/StateTreeComponent.cpp
+++ b/Engine/Plugins/Runtime/GameplayStateTree/Source/GameplayStateTreeModule/Private/Components/StateTreeComponent.cpp
@@ -123,6 +123,7 @@ void UStateTreeComponent::TickComponent(float DeltaTime, enum ELevelTick TickTyp
 	if (!StateTreeRef.IsValid())
 	{
 		STATETREE_LOG(Warning, TEXT("%hs: Trying to tick State Tree component with invalid asset."), __FUNCTION__);
+		bIsRunning = false;
 		DisableTick();
 		return;
 	}
@@ -151,6 +152,7 @@ void UStateTreeComponent::TickComponent(float DeltaTime, enum ELevelTick TickTyp
 	}
 	else
 	{
+		STATETREE_LOG(Warning, TEXT("Context Requirements in %hs failed. Component tick is disabled."), __FUNCTION__);
 		DisableTick();
 	}
 }
@@ -169,10 +171,9 @@ void UStateTreeComponent::RestartLogic()
 
 void UStateTreeComponent::StartTree()
 {
-	bIsRunning = false;
-
 	if (HasValidStateTreeReference().HasError())
 	{
+		bIsRunning = false;
 		DisableTick();
 		return;
 	}
@@ -206,6 +207,7 @@ void UStateTreeComponent::StartTree()
 	}
 	else
 	{
+		STATETREE_LOG(Warning, TEXT("Context Requirements in %hs failed. Component tick is disabled."), __FUNCTION__);
 		DisableTick();
 	}
 }
@@ -216,6 +218,7 @@ void UStateTreeComponent::StopLogic(const FString& Reason)
 
 	auto StopTree = [this](FStateTreeExecutionContext& Context)
 	{
+		bIsRunning = false;
 		const EStateTreeRunStatus PreviousRunStatus = Context.GetStateTreeRunStatus();
 		const EStateTreeRunStatus CurrentRunStatus = Context.Stop();
 
@@ -231,11 +234,11 @@ void UStateTreeComponent::StopLogic(const FString& Reason)
 		return;
 	}
 
-	bIsRunning = false;
 	DisableTick();
 
-	if (!StateTreeRef.IsValid())
+	if (HasValidStateTreeReference().HasError())
 	{
+		bIsRunning = false;
 		STATETREE_LOG(Warning, TEXT("%hs: Trying to stop State Tree component with invalid asset."), __FUNCTION__);
 		return;
 	}
@@ -247,11 +250,15 @@ void UStateTreeComponent::StopLogic(const FString& Reason)
 	}
 	else
 	{
-	FStateTreeExecutionContext Context(*GetOwner(), *StateTreeRef.GetStateTree(), InstanceData);
-	if (SetContextRequirements(Context))
-	{
+		FStateTreeExecutionContext Context(*GetOwner(), *StateTreeRef.GetStateTree(), InstanceData);
+		if (SetContextRequirements(Context))
+		{
 			StopTree(Context);
 		}
+		else
+		{
+			STATETREE_LOG(Warning, TEXT("Context Requirements in %hs failed. Component tick is disabled."), __FUNCTION__);
+		}
 	}
 }
 
diff --git a/Engine/Plugins/Runtime/GameplayStateTree/Source/GameplayStateTreeModule/Public/Components/StateTreeComponent.h b/Engine/Plugins/Runtime/GameplayStateTree/Source/GameplayStateTreeModule/Public/Components/StateTreeComponent.h
--- a/Engine/Plugins/Runtime/GameplayStateTree/Source/GameplayStateTreeModule/Public/Components/StateTreeComponent.h
+++ b/Engine/Plugins/Runtime/GameplayStateTree/Source/GameplayStateTreeModule/Public/Components/StateTreeComponent.h
@@ -186,10 +186,10 @@ protected:
 	UPROPERTY(EditAnywhere, Category = AI)
 	bool bStartLogicAutomatically = true;
 	
-	/** if set, state tree execution is allowed */
+	/** if set, execution has started and has not stopped yet. */
 	uint8 bIsRunning : 1;
 	
-	/** if set, execution requests will be postponed */
+	/** if set, execution has been requested to stop ticking. */
 	uint8 bIsPaused : 1;
 
 private:
