PR #13642: Added UE_API macro to FContextDataSetter

diff --git a/Engine/Plugins/Runtime/GameplayStateTree/Source/GameplayStateTreeModule/Public/Components/StateTreeComponentSchema.h b/Engine/Plugins/Runtime/GameplayStateTree/Source/GameplayStateTreeModule/Public/Components/StateTreeComponentSchema.h
--- a/Engine/Plugins/Runtime/GameplayStateTree/Source/GameplayStateTreeModule/Public/Components/StateTreeComponentSchema.h
+++ b/Engine/Plugins/Runtime/GameplayStateTree/Source/GameplayStateTreeModule/Public/Components/StateTreeComponentSchema.h
@@ -48,16 +48,16 @@ protected:
 	struct FContextDataSetter
 	{
 	public:
-		FContextDataSetter(TNotNull<const UBrainComponent*> BrainComponent, FStateTreeExecutionContext& Context);
+		UE_API FContextDataSetter(TNotNull<const UBrainComponent*> BrainComponent, FStateTreeExecutionContext& Context);
 
 		TNotNull<const UBrainComponent*> GetComponent() const
 		{
 			return BrainComponent;
 		}
-		TNotNull<const UStateTree*> GetStateTree() const;
-		TNotNull<const UStateTreeComponentSchema*> GetSchema() const;
+		UE_API TNotNull<const UStateTree*> GetStateTree() const;
+		UE_API TNotNull<const UStateTreeComponentSchema*> GetSchema() const;
 
-		bool SetContextDataByName(FName Name, FStateTreeDataView DataView);
+		UE_API bool SetContextDataByName(FName Name, FStateTreeDataView DataView);
 
 	private:
 		TNotNull<const UBrainComponent*> BrainComponent;
