[State Tree] use Readonly Context to validate State Tree reference

diff --git a/Engine/Plugins/Runtime/GameplayStateTree/Source/GameplayStateTreeModule/Private/Components/StateTreeComponent.cpp b/Engine/Plugins/Runtime/GameplayStateTree/Source/GameplayStateTreeModule/Private/Components/StateTreeComponent.cpp
--- a/Engine/Plugins/Runtime/GameplayStateTree/Source/GameplayStateTreeModule/Private/Components/StateTreeComponent.cpp
+++ b/Engine/Plugins/Runtime/GameplayStateTree/Source/GameplayStateTreeModule/Private/Components/StateTreeComponent.cpp
@@ -438,7 +438,7 @@ void UStateTreeComponent::ValidateStateTreeReference()
  	 	return;
  	}
 
-	const FStateTreeExecutionContext Context(*GetOwner(), *StateTreeRef.GetStateTree(), InstanceData);
+	FStateTreeReadOnlyExecutionContext Context(GetOwner(), StateTreeRef.GetStateTree(), InstanceData);
 	if (!Context.IsValid())
 	{
 		STATETREE_LOG(Error, TEXT("%hs: Failed to init StateTreeContext."), __FUNCTION__);
