[StateTree] Ensure the StateTree asset is compiled before using it in a StateTreeComponent

#Synced-CL 51260989

diff --git a/Engine/Plugins/Runtime/GameplayStateTree/Source/GameplayStateTreeModule/Private/Components/StateTreeComponent.cpp b/Engine/Plugins/Runtime/GameplayStateTree/Source/GameplayStateTreeModule/Private/Components/StateTreeComponent.cpp
--- a/Engine/Plugins/Runtime/GameplayStateTree/Source/GameplayStateTreeModule/Private/Components/StateTreeComponent.cpp
+++ b/Engine/Plugins/Runtime/GameplayStateTree/Source/GameplayStateTreeModule/Private/Components/StateTreeComponent.cpp
@@ -453,6 +453,11 @@ TValueOrError<void, FString> UStateTreeComponent::HasValidStateTreeReference() c
 		return MakeError(TEXT("The State Tree asset is not set."));
 	}
 
+	if (!StateTreeRef.GetStateTree()->IsReadyToRun())
+	{
+		return MakeError(TEXT("The State Tree schema is not ready to run."));
+	}
+
 	if (StateTreeRef.GetStateTree()->GetSchema() == nullptr
 		|| !StateTreeRef.GetStateTree()->GetSchema()->GetClass()->IsChildOf(UStateTreeComponentSchema::StaticClass()))
 	{
diff --git a/Engine/Plugins/Runtime/StateTree/Source/StateTreeEditorModule/Private/StateTreeCompilerManager.cpp b/Engine/Plugins/Runtime/StateTree/Source/StateTreeEditorModule/Private/StateTreeCompilerManager.cpp
--- a/Engine/Plugins/Runtime/StateTree/Source/StateTreeEditorModule/Private/StateTreeCompilerManager.cpp
+++ b/Engine/Plugins/Runtime/StateTree/Source/StateTreeEditorModule/Private/StateTreeCompilerManager.cpp
@@ -297,7 +297,10 @@ FCompilerManagerImpl::FCompilerManagerImpl()
 	EndPIEHandle = FEditorDelegates::EndPIE.AddRaw(this, &FCompilerManagerImpl::HandleEndPIE);
 
 	// When running in standalone
-	bCanQueueCompilation = !IsRunningGame();
+	if (IsRunningGame())
+	{
+		bCanQueueCompilation = !Private::bEnableCompileSynchronouslyInPIESession;
+	}
 }
 
 FCompilerManagerImpl::~FCompilerManagerImpl()
